RDG Insights Plugin

diff --git a/Engine/Plugins/RenderGraphInsights/RenderGraphInsights.uplugin b/Engine/Plugins/RenderGraphInsights/RenderGraphInsights.uplugin
new file mode 100644
--- /dev/null
+++ b/Engine/Plugins/RenderGraphInsights/RenderGraphInsights.uplugin
@@ -0,0 +1,28 @@
+{
+	"FileVersion": 3,
+	"Version": 1,
+	"VersionName": "1.0",
+	"FriendlyName": "RDG Insights",
+	"Description": "Allows debugging of RDG via Unreal Insights",
+	"Category": "Insights",
+	"CreatedBy": "Epic Games, Inc.",
+	"CreatedByURL": "http://epicgames.com",
+	"DocsURL": "",
+	"MarketplaceURL": "",
+	"SupportURL": "",
+	"EnabledByDefault": false,
+	"CanContainContent": false,
+	"IsBetaVersion": false,
+	"Installed": false,
+	"SupportedPrograms": [ "UnrealInsights" ],
+	"Modules": [
+		{
+			"Name": "RenderGraphInsights",
+			"Type": "EditorAndProgram",
+			"LoadingPhase": "Default",
+			"WhitelistPrograms": [
+				"UnrealInsights"
+			]
+		}
+	]
+}
diff --git a/Engine/Plugins/RenderGraphInsights/Source/RenderGraphInsights/Private/RenderGraphAnalyzer.cpp b/Engine/Plugins/RenderGraphInsights/Source/RenderGraphInsights/Private/RenderGraphAnalyzer.cpp
new file mode 100644
--- /dev/null
+++ b/Engine/Plugins/RenderGraphInsights/Source/RenderGraphInsights/Private/RenderGraphAnalyzer.cpp
@@ -0,0 +1,67 @@
+// Copyright Epic Games, Inc. All Rights Reserved.
+
+#include "RenderGraphAnalyzer.h"
+
+#include "RenderGraphProvider.h"
+#include "Trace/Analyzer.h"
+#include "TraceServices/Model/AnalysisSession.h"
+
+#define LOCTEXT_NAMESPACE "RenderGraphProvider"
+
+namespace UE
+{
+namespace RenderGraphInsights
+{
+
+FRenderGraphAnalyzer::FRenderGraphAnalyzer(TraceServices::IAnalysisSession& InSession, FRenderGraphProvider& InRenderGraphProvider)
+	: Session(InSession)
+	, Provider(InRenderGraphProvider)
+{}
+
+void FRenderGraphAnalyzer::OnAnalysisBegin(const FOnAnalysisContext& Context)
+{
+	UE::Trace::IAnalyzer::FInterfaceBuilder& Builder = Context.InterfaceBuilder;
+
+	Builder.RouteEvent(RouteId_Graph, "RDGTrace", "GraphMessage");
+	Builder.RouteEvent(RouteId_GraphEnd, "RDGTrace", "GraphEndMessage");
+	Builder.RouteEvent(RouteId_Scope, "RDGTrace", "ScopeMessage");
+	Builder.RouteEvent(RouteId_Pass, "RDGTrace", "PassMessage");
+	Builder.RouteEvent(RouteId_Texture, "RDGTrace", "TextureMessage");
+	Builder.RouteEvent(RouteId_Buffer, "RDGTrace", "BufferMessage");
+}
+
+bool FRenderGraphAnalyzer::OnEvent(uint16 RouteId, EStyle Style, const FOnEventContext& Context)
+{
+	TraceServices::FAnalysisSessionEditScope _(Session);
+	switch (RouteId)
+	{
+	case RouteId_Graph:
+	{
+		double EndTime{};
+		Provider.AddGraph(Context, EndTime);
+		Session.UpdateDurationSeconds(EndTime);
+		break;
+	}
+	case RouteId_GraphEnd:
+		Provider.AddGraphEnd();
+		break;
+	case RouteId_Scope:
+		Provider.AddScope(FScopePacket(Context));
+		break;
+	case RouteId_Pass:
+		Provider.AddPass(FPassPacket(Context));
+		break;
+	case RouteId_Texture:
+		Provider.AddTexture(FTexturePacket(Context));
+		break;
+	case RouteId_Buffer:
+		Provider.AddBuffer(FBufferPacket(Context));
+		break;
+	}
+	return true;
+}
+
+} //namespace RenderGraphInsights
+} //namespace UE
+
+#undef LOCTEXT_NAMESPACE
\ No newline at end of file
diff --git a/Engine/Plugins/RenderGraphInsights/Source/RenderGraphInsights/Private/RenderGraphAnalyzer.h b/Engine/Plugins/RenderGraphInsights/Source/RenderGraphInsights/Private/RenderGraphAnalyzer.h
new file mode 100644
--- /dev/null
+++ b/Engine/Plugins/RenderGraphInsights/Source/RenderGraphInsights/Private/RenderGraphAnalyzer.h
@@ -0,0 +1,41 @@
+// Copyright Epic Games, Inc. All Rights Reserved.
+
+#pragma once
+
+#include "Trace/Analyzer.h"
+
+namespace TraceServices { class IAnalysisSession; }
+
+namespace UE
+{
+namespace RenderGraphInsights
+{ 
+
+class FRenderGraphProvider;
+
+class FRenderGraphAnalyzer : public UE::Trace::IAnalyzer
+{
+public:
+	FRenderGraphAnalyzer(TraceServices::IAnalysisSession& InSession, FRenderGraphProvider& InRenderGraphProvider);
+
+	virtual void OnAnalysisBegin(const FOnAnalysisContext& Context) override;
+	virtual void OnAnalysisEnd() override {}
+	virtual bool OnEvent(uint16 RouteId, EStyle Style, const FOnEventContext& Context) override;
+
+private:
+	enum : uint16
+	{
+		RouteId_Graph,
+		RouteId_GraphEnd,
+		RouteId_Scope,
+		RouteId_Pass,
+		RouteId_Buffer,
+		RouteId_Texture
+	};
+
+	TraceServices::IAnalysisSession& Session;
+	FRenderGraphProvider& Provider;
+};
+
+} //namespace RenderGraphInsights
+} //namespace UE
diff --git a/Engine/Plugins/RenderGraphInsights/Source/RenderGraphInsights/Private/RenderGraphInsightsModule.cpp b/Engine/Plugins/RenderGraphInsights/Source/RenderGraphInsights/Private/RenderGraphInsightsModule.cpp
new file mode 100644
--- /dev/null
+++ b/Engine/Plugins/RenderGraphInsights/Source/RenderGraphInsights/Private/RenderGraphInsightsModule.cpp
@@ -0,0 +1,35 @@
+// Copyright Epic Games, Inc. All Rights Reserved.
+
+#include "RenderGraphInsightsModule.h"
+#include "Modules/ModuleManager.h"
+#include "Features/IModularFeatures.h"
+
+#include "RenderGraphTraceModule.h"
+#include "RenderGraphTimingViewExtender.h"
+
+namespace UE
+{
+namespace RenderGraphInsights
+{
+
+FRenderGraphInsightsModule& FRenderGraphInsightsModule::Get()
+{
+	return FModuleManager::LoadModuleChecked<FRenderGraphInsightsModule>("RenderGraphInsights");
+}
+
+void FRenderGraphInsightsModule::StartupModule()
+{
+	IModularFeatures::Get().RegisterModularFeature(TraceServices::ModuleFeatureName, &TraceModule);
+	IModularFeatures::Get().RegisterModularFeature(Insights::TimingViewExtenderFeatureName, &TimingViewExtender);
+}
+
+void FRenderGraphInsightsModule::ShutdownModule()
+{
+	IModularFeatures::Get().UnregisterModularFeature(TraceServices::ModuleFeatureName, &TraceModule);
+	IModularFeatures::Get().UnregisterModularFeature(Insights::TimingViewExtenderFeatureName, &TimingViewExtender);
+}
+
+} //namespace SlateInsights
+} //namespace UE
+
+IMPLEMENT_MODULE((UE::RenderGraphInsights::FRenderGraphInsightsModule), RenderGraphInsights);
\ No newline at end of file
diff --git a/Engine/Plugins/RenderGraphInsights/Source/RenderGraphInsights/Private/RenderGraphInsightsModule.h b/Engine/Plugins/RenderGraphInsights/Source/RenderGraphInsights/Private/RenderGraphInsightsModule.h
new file mode 100644
--- /dev/null
+++ b/Engine/Plugins/RenderGraphInsights/Source/RenderGraphInsights/Private/RenderGraphInsightsModule.h
@@ -0,0 +1,35 @@
+// Copyright Epic Games, Inc. All Rights Reserved.
+
+#pragma once
+
+#include "CoreMinimal.h"
+#include "Modules/ModuleInterface.h"
+
+#include "RenderGraphTraceModule.h"
+#include "RenderGraphTimingViewExtender.h"
+
+class FTabManager;
+
+namespace UE
+{
+namespace RenderGraphInsights
+{
+
+class FRenderGraphInsightsModule : public IModuleInterface
+{
+public:
+	static FRenderGraphInsightsModule& Get();
+
+	//~ Begin IModuleInterface interface
+	virtual void StartupModule() override;
+	virtual void ShutdownModule() override;
+	//~ End IModuleInterface interface
+
+private:
+	FRenderGraphTraceModule TraceModule;
+	FRenderGraphTimingViewExtender TimingViewExtender;
+	TWeakPtr<FTabManager> InsightsTabManager;
+};
+
+} //namespace RenderGraphInsights
+} //namespace UE
\ No newline at end of file
diff --git a/Engine/Plugins/RenderGraphInsights/Source/RenderGraphInsights/Private/RenderGraphProvider.cpp b/Engine/Plugins/RenderGraphInsights/Source/RenderGraphInsights/Private/RenderGraphProvider.cpp
new file mode 100644
--- /dev/null
+++ b/Engine/Plugins/RenderGraphInsights/Source/RenderGraphInsights/Private/RenderGraphProvider.cpp
@@ -0,0 +1,267 @@
+// Copyright Epic Games, Inc. All Rights Reserved.
+
+#include "RenderGraphProvider.h"
+
+namespace UE
+{
+namespace RenderGraphInsights
+{
+
+INSIGHTS_IMPLEMENT_RTTI(FPacket)
+INSIGHTS_IMPLEMENT_RTTI(FScopePacket)
+INSIGHTS_IMPLEMENT_RTTI(FResourcePacket)
+INSIGHTS_IMPLEMENT_RTTI(FTexturePacket)
+INSIGHTS_IMPLEMENT_RTTI(FBufferPacket)
+INSIGHTS_IMPLEMENT_RTTI(FPassPacket)
+INSIGHTS_IMPLEMENT_RTTI(FGraphPacket)
+INSIGHTS_IMPLEMENT_RTTI(FPassIntervalPacket)
+
+FName FRenderGraphProvider::ProviderName("RenderGraphProvider");
+
+void SanitizeName(FString& Name)
+{
+	if (Name.IsEmpty())
+	{
+		Name = TEXT("<unnamed>");
+	}
+}
+
+FString GetSanitizeName(FString Name)
+{
+	SanitizeName(Name);
+	return MoveTemp(Name);
+}
+
+FString GetSizeName(uint32 Bytes)
+{
+	const uint32 KB = 1024;
+	const uint32 MB = 1024 * 1024;
+
+	if (Bytes < MB)
+	{
+		return FString::Printf(TEXT(" (%.3fKB)"), (float)Bytes / (float)KB);
+	}
+	else
+	{
+		return FString::Printf(TEXT(" (%.3fMB)"), (float)Bytes / (float)MB);
+	}
+}
+
+FStringView GetString(const UE::Trace::IAnalyzer::FEventData& EventData, const ANSICHAR* FieldName)
+{
+	FStringView Value;
+	EventData.GetString(FieldName, Value);
+	return Value;
+}
+
+FPacket::FPacket(const UE::Trace::IAnalyzer::FOnEventContext& Context)
+	: Name(GetString(Context.EventData, "Name"))
+	, StartTime(Context.EventTime.AsSeconds(Context.EventData.GetValue<uint64>("StartCycles")))
+	, EndTime(Context.EventTime.AsSeconds(Context.EventData.GetValue<uint64>("EndCycles")))
+{}
+
+FPassIntervalPacket::FPassIntervalPacket(const UE::Trace::IAnalyzer::FOnEventContext& Context)
+	: FPacket(Context)
+	, FirstPass(Context.EventData.GetValue<FRDGPassHandle>("FirstPass"))
+	, LastPass(Context.EventData.GetValue<FRDGPassHandle>("LastPass"))
+{}
+
+FScopePacket::FScopePacket(const UE::Trace::IAnalyzer::FOnEventContext& Context)
+	: FPassIntervalPacket(Context)
+	, Depth(Context.EventData.GetValue<uint16>("Depth"))
+{}
+
+FResourcePacket::FResourcePacket(const UE::Trace::IAnalyzer::FOnEventContext& Context)
+	: FPassIntervalPacket(Context)
+	, Passes(Context.EventData.GetArrayView<FRDGPassHandle>("Passes"))
+	, bExternal(Context.EventData.GetValue<bool>("IsExternal"))
+	, bExtracted(Context.EventData.GetValue<bool>("IsExtracted"))
+	, bCulled(Context.EventData.GetValue<bool>("IsCulled"))
+{
+	if (Passes.Num())
+	{
+		FirstPass = Passes[0];
+		LastPass  = Passes.Last();
+	}
+}
+
+FTexturePacket::FTexturePacket(const UE::Trace::IAnalyzer::FOnEventContext& Context)
+	: FResourcePacket(Context)
+	, Handle(Context.EventData.GetValue<FRDGTextureHandle>("Handle"))
+	, NextOwnerHandle(Context.EventData.GetValue<FRDGTextureHandle>("NextOwnerHandle"))
+	, SizeInBytes(Context.EventData.GetValue<uint64>("SizeInBytes"))
+{
+	Desc.Flags = ETextureCreateFlags(Context.EventData.GetValue<uint64>("CreateFlags"));
+	Desc.Dimension = ETextureDimension(Context.EventData.GetValue<uint16>("Dimension"));
+	Desc.Format = EPixelFormat(Context.EventData.GetValue<uint16>("Format"));
+	Desc.Extent.X = Context.EventData.GetValue<uint32>("ExtentX");
+	Desc.Extent.Y = Context.EventData.GetValue<uint32>("ExtentY");
+	Desc.Depth = Context.EventData.GetValue<uint16>("Depth");
+	Desc.ArraySize = Context.EventData.GetValue<uint16>("ArraySize");
+	Desc.NumMips = Context.EventData.GetValue<uint8>("NumMips");
+	Desc.NumSamples = Context.EventData.GetValue<uint8>("NumSamples");
+
+	Name += GetSizeName(SizeInBytes);
+}
+
+FBufferPacket::FBufferPacket(const UE::Trace::IAnalyzer::FOnEventContext& Context)
+	: FResourcePacket(Context)
+	, Handle(Context.EventData.GetValue<FRDGBufferHandle>("Handle"))
+	, NextOwnerHandle(Context.EventData.GetValue<FRDGBufferHandle>("NextOwnerHandle"))
+{
+	Desc.Usage = EBufferUsageFlags(Context.EventData.GetValue<uint32>("UsageFlags"));
+	Desc.BytesPerElement = Context.EventData.GetValue<uint32>("BytesPerElement");
+	Desc.NumElements = Context.EventData.GetValue<uint32>("NumElements");
+
+	const uint64 SizeInBytes = Desc.BytesPerElement * Desc.NumElements;
+	Name += GetSizeName(SizeInBytes);
+}
+
+FPassPacket::FPassPacket(const UE::Trace::IAnalyzer::FOnEventContext& Context)
+	: FPacket(Context)
+	, Textures(Context.EventData.GetArrayView<FRDGTextureHandle>("Textures"))
+	, Buffers(Context.EventData.GetArrayView<FRDGBufferHandle>("Buffers"))
+	, Handle(Context.EventData.GetValue<FRDGPassHandle>("Handle"))
+	, GraphicsForkPass(Context.EventData.GetValue<FRDGPassHandle>("GraphicsForkPass"))
+	, GraphicsJoinPass(Context.EventData.GetValue<FRDGPassHandle>("GraphicsJoinPass"))
+	, Flags(ERDGPassFlags(Context.EventData.GetValue<uint16>("Flags")))
+	, Pipeline(ERHIPipeline(Context.EventData.GetValue<uint8>("Pipeline")))
+	, bCulled(Context.EventData.GetValue<bool>("IsCulled"))
+	, bAsyncComputeBegin(Context.EventData.GetValue<bool>("IsAsyncComputeBegin"))
+	, bAsyncComputeEnd(Context.EventData.GetValue<bool>("IsAsyncComputeEnd"))
+	, bSkipRenderPassBegin(Context.EventData.GetValue<bool>("SkipRenderPassBegin"))
+	, bSkipRenderPassEnd(Context.EventData.GetValue<bool>("SkipRenderPassEnd"))
+{}
+
+static const uint64 PageSize = 1024;
+
+FGraphPacket::FGraphPacket(TraceServices::ILinearAllocator& Allocator, const UE::Trace::IAnalyzer::FOnEventContext& Context)
+	: FPacket(Context)
+	, Scopes(Allocator, PageSize)
+	, Passes(Allocator, PageSize)
+	, Textures(Allocator, PageSize)
+	, Buffers(Allocator, PageSize)
+	, PassCount(Context.EventData.GetValue<uint16>("PassCount"))
+{
+	NormalizedPassDuration = (EndTime - StartTime) / double(PassCount);
+}
+
+FRenderGraphProvider::FRenderGraphProvider(TraceServices::IAnalysisSession& InSession)
+	: Session(InSession)
+	, GraphTimeline(Session.GetLinearAllocator())
+{}
+
+void FRenderGraphProvider::AddGraph(const UE::Trace::IAnalyzer::FOnEventContext& Context, double& OutEndTime)
+{
+	CurrentGraph = MakeShared<FGraphPacket>(Session.GetLinearAllocator(), Context);
+	SanitizeName(CurrentGraph->Name);
+
+	OutEndTime = CurrentGraph->EndTime;
+}
+
+void FRenderGraphProvider::AddGraphEnd()
+{
+	const double EndTime = CurrentGraph->EndTime;
+	const uint64 EventId = GraphTimeline.EmplaceBeginEvent(CurrentGraph->StartTime, MoveTemp(CurrentGraph));
+	GraphTimeline.EndEvent(EventId, EndTime);
+}
+
+void FRenderGraphProvider::AddScope(FScopePacket InScope)
+{
+	const FPassPacket& FirstPass = *CurrentGraph->GetPass(InScope.FirstPass);
+	const FPassPacket& LastPass  = *CurrentGraph->GetPass(InScope.LastPass);
+
+	InScope.Graph = CurrentGraph.Get();
+	InScope.StartTime = FirstPass.StartTime;
+	InScope.EndTime = LastPass.EndTime;
+
+	CurrentGraph->ScopeDepth = FMath::Max(CurrentGraph->ScopeDepth, InScope.Depth);
+	CurrentGraph->Scopes.EmplaceBack(InScope);
+}
+
+void FRenderGraphProvider::AddPass(FPassPacket InPass)
+{
+	InPass.Graph = CurrentGraph.Get();
+
+	const double NormalizedPassDuration = CurrentGraph->NormalizedPassDuration;
+	const uint32 PassIndex = CurrentGraph->Passes.Num();
+	InPass.StartTime = CurrentGraph->StartTime + NormalizedPassDuration * double(PassIndex);
+	InPass.EndTime = InPass.StartTime + NormalizedPassDuration;
+
+	CurrentGraph->Passes.EmplaceBack(InPass);
+}
+
+void FRenderGraphProvider::SetupResource(FResourcePacket& Resource)
+{
+	Resource.Graph = CurrentGraph.Get();
+
+	if (Resource.bCulled)
+	{
+		return;
+	}
+
+	const FPassPacket* FirstPass = CurrentGraph->GetPass(Resource.FirstPass);
+	const FPassPacket* LastPass  = CurrentGraph->GetPass(Resource.LastPass);
+
+	if (Resource.bExternal)
+	{
+		FirstPass = CurrentGraph->GetProloguePass();
+	}
+
+	if (Resource.bExtracted)
+	{
+		LastPass = CurrentGraph->GetEpiloguePass();
+	}
+
+	Resource.StartTime = FirstPass->StartTime;
+	Resource.EndTime = LastPass->EndTime;
+}
+
+void FRenderGraphProvider::AddTexture(const FTexturePacket InTexture)
+{
+	FTexturePacket& Texture = CurrentGraph->Textures.EmplaceBack(InTexture);
+	Texture.Index = CurrentGraph->Textures.Num() - 1;
+	SetupResource(Texture);
+
+	if (const FTexturePacket* const* PreviousOwnerPtr = CurrentGraph->TextureHandleToPreviousOwner.Find(Texture.Handle))
+	{
+		const FTexturePacket* PreviousOwner = *PreviousOwnerPtr;
+		Texture.PreviousOwner = PreviousOwner;
+		Texture.Depth = PreviousOwner->Depth;
+	}
+	else if (!Texture.bCulled)
+	{
+		Texture.Depth = CurrentGraph->TextureDepthOffset++;
+	}
+
+	if (Texture.NextOwnerHandle.IsValid())
+	{
+		CurrentGraph->TextureHandleToPreviousOwner.Emplace(Texture.NextOwnerHandle, &Texture);
+	}
+}
+
+void FRenderGraphProvider::AddBuffer(const FBufferPacket InBuffer)
+{
+	FBufferPacket& Buffer = CurrentGraph->Buffers.EmplaceBack(InBuffer);
+	Buffer.Index = CurrentGraph->Buffers.Num() - 1;
+	SetupResource(Buffer);
+
+	if (const FBufferPacket* const* PreviousOwnerPtr = CurrentGraph->BufferHandleToPreviousOwner.Find(Buffer.Handle))
+	{
+		const FBufferPacket* PreviousOwner = *PreviousOwnerPtr;
+		Buffer.PreviousOwner = PreviousOwner;
+		Buffer.Depth = PreviousOwner->Depth;
+	}
+	else if (!Buffer.bCulled)
+	{
+		Buffer.Depth = CurrentGraph->BufferDepthOffset++;
+	}
+
+	if (Buffer.NextOwnerHandle.IsValid())
+	{
+		CurrentGraph->BufferHandleToPreviousOwner.Emplace(Buffer.NextOwnerHandle, &Buffer);
+	}
+}
+
+} //namespace RenderGraphInsights
+} //namespace UE
\ No newline at end of file
diff --git a/Engine/Plugins/RenderGraphInsights/Source/RenderGraphInsights/Private/RenderGraphProvider.h b/Engine/Plugins/RenderGraphInsights/Source/RenderGraphInsights/Private/RenderGraphProvider.h
new file mode 100644
--- /dev/null
+++ b/Engine/Plugins/RenderGraphInsights/Source/RenderGraphInsights/Private/RenderGraphProvider.h
@@ -0,0 +1,200 @@
+// Copyright Epic Games, Inc. All Rights Reserved.
+
+#pragma once
+
+#include "CoreMinimal.h"
+#include "TraceServices/Model/AnalysisSession.h"
+#include "Insights/ViewModels/TimingEvent.h"
+#include "Common/PagedArray.h"
+#include "Model/IntervalTimeline.h"
+#include "Trace/Analyzer.h"
+#include "RenderGraphPass.h"
+
+namespace TraceServices { class IAnalysisSession; }
+
+namespace UE
+{
+namespace RenderGraphInsights
+{
+
+class FGraphPacket;
+
+const uint32 kInvalidVisibleIndex = ~0u;
+
+class FPacket
+{
+	INSIGHTS_DECLARE_RTTI_BASE(FPacket);
+public:
+	virtual ~FPacket() = default;
+
+	const FGraphPacket* Graph{};
+	FString Name;
+	double StartTime{};
+	double EndTime{};
+	mutable uint32 VisibleIndex = kInvalidVisibleIndex;
+
+	FPacket(const UE::Trace::IAnalyzer::FOnEventContext& Context);
+};
+
+class FPassIntervalPacket : public FPacket
+{
+	INSIGHTS_DECLARE_RTTI(FPassIntervalPacket, FPacket)
+public:
+	FRDGPassHandle FirstPass;
+	FRDGPassHandle LastPass;
+
+	FPassIntervalPacket(const UE::Trace::IAnalyzer::FOnEventContext& Context);
+};
+
+class FScopePacket : public FPassIntervalPacket
+{
+	INSIGHTS_DECLARE_RTTI(FScopePacket, FPassIntervalPacket)
+public:
+	uint32 Depth{};
+
+	FScopePacket(const UE::Trace::IAnalyzer::FOnEventContext& Context);
+};
+
+class FResourcePacket : public FPassIntervalPacket
+{
+	INSIGHTS_DECLARE_RTTI(FResourcePacket, FPassIntervalPacket)
+public:
+	uint32 Index{};
+	uint32 Depth{};
+
+	TArray<FRDGPassHandle> Passes;
+
+	bool bExternal{};
+	bool bExtracted{};
+	bool bCulled{};
+
+	FResourcePacket(const UE::Trace::IAnalyzer::FOnEventContext& Context);
+};
+
+class FTexturePacket : public FResourcePacket
+{
+	INSIGHTS_DECLARE_RTTI(FTexturePacket, FResourcePacket)
+public:
+	const FTexturePacket* PreviousOwner{};
+
+	FRDGTextureHandle Handle;
+	FRDGTextureHandle NextOwnerHandle;
+	FRDGTextureDesc Desc;
+	uint64 SizeInBytes{};
+
+	FTexturePacket(const UE::Trace::IAnalyzer::FOnEventContext& Context);
+};
+
+class FBufferPacket : public FResourcePacket
+{
+	INSIGHTS_DECLARE_RTTI(FBufferPacket, FResourcePacket)
+public:
+	const FBufferPacket* PreviousOwner{};
+
+	FRDGBufferHandle Handle;
+	FRDGBufferHandle NextOwnerHandle;
+	FRDGBufferDesc Desc;
+
+	FBufferPacket(const UE::Trace::IAnalyzer::FOnEventContext& Context);
+};
+
+class FPassPacket : public FPacket
+{
+	INSIGHTS_DECLARE_RTTI(FPassPacket, FPacket)
+public:
+	TArray<FRDGTextureHandle> Textures;
+	TArray<FRDGBufferHandle> Buffers;
+	FRDGPassHandle Handle;
+	FRDGPassHandle GraphicsForkPass;
+	FRDGPassHandle GraphicsJoinPass;
+	ERDGPassFlags Flags{};
+	ERHIPipeline Pipeline{};
+	bool bCulled{};
+	bool bAsyncComputeBegin{};
+	bool bAsyncComputeEnd{};
+	bool bSkipRenderPassBegin{};
+	bool bSkipRenderPassEnd{};
+
+	FPassPacket(const UE::Trace::IAnalyzer::FOnEventContext& Context);
+};
+
+class FGraphPacket : public FPacket
+{
+	INSIGHTS_DECLARE_RTTI(FGraphPacket, FPacket)
+public:
+	double NormalizedPassDuration{};
+
+	TraceServices::TPagedArray<FScopePacket> Scopes;
+	TraceServices::TPagedArray<FPassPacket> Passes;
+	TraceServices::TPagedArray<FTexturePacket> Textures;
+	TraceServices::TPagedArray<FBufferPacket> Buffers;
+
+	TMap<FRDGBufferHandle, const FBufferPacket*> BufferHandleToPreviousOwner;
+	TMap<FRDGTextureHandle, const FTexturePacket*> TextureHandleToPreviousOwner;
+
+	uint32 ScopeDepth{};
+	uint32 PassCount{};
+	uint32 TextureDepthOffset{};
+	uint32 BufferDepthOffset{};
+
+	FGraphPacket(TraceServices::ILinearAllocator& Allocator, const UE::Trace::IAnalyzer::FOnEventContext& Context);
+
+	const FPassPacket* GetProloguePass() const
+	{
+		return &Passes[0];
+	}
+
+	const FPassPacket* GetEpiloguePass() const
+	{
+		return &Passes.Last();
+	}
+
+	const FPassPacket* GetPass(FRDGPassHandle Handle) const
+	{
+		return &Passes[Handle.GetIndex()];
+	}
+
+	const FTexturePacket* GetTexture(FRDGTextureHandle Handle) const
+	{
+		return &Textures[Handle.GetIndex()];
+	}
+
+	const FBufferPacket* GetBuffer(FRDGBufferHandle Handle) const
+	{
+		return &Buffers[Handle.GetIndex()];
+	}
+};
+
+class FRenderGraphProvider : public TraceServices::IProvider
+{
+public:
+	static FName ProviderName;
+
+	FRenderGraphProvider(TraceServices::IAnalysisSession& InSession);
+
+	void AddGraph(const UE::Trace::IAnalyzer::FOnEventContext& Context, double& OutEndTime);
+	void AddGraphEnd();
+	void AddScope(FScopePacket Scope);
+	void AddPass(FPassPacket Pass);
+	void AddTexture(FTexturePacket Texture);
+	void AddBuffer(FBufferPacket Buffer);
+
+	using TGraphTimeline = TraceServices::TIntervalTimeline<TSharedPtr<FGraphPacket>>;
+	const TGraphTimeline& GetGraphTimeline() const
+	{
+		Session.ReadAccessCheck();
+		return GraphTimeline;
+	}
+
+private:
+	void SetupResource(FResourcePacket& ResourcePacket);
+
+	TraceServices::IAnalysisSession& Session;
+
+	TSharedPtr<FGraphPacket> CurrentGraph;
+
+	TGraphTimeline GraphTimeline;
+};
+
+} //namespace RenderGraphInsights
+} //namespace UE
diff --git a/Engine/Plugins/RenderGraphInsights/Source/RenderGraphInsights/Private/RenderGraphTimingViewExtender.cpp b/Engine/Plugins/RenderGraphInsights/Source/RenderGraphInsights/Private/RenderGraphTimingViewExtender.cpp
new file mode 100644
--- /dev/null
+++ b/Engine/Plugins/RenderGraphInsights/Source/RenderGraphInsights/Private/RenderGraphTimingViewExtender.cpp
@@ -0,0 +1,63 @@
+// Copyright Epic Games, Inc. All Rights Reserved.
+
+#include "RenderGraphTimingViewExtender.h"
+#include "Insights/ITimingViewSession.h"
+#include "UObject/WeakObjectPtr.h"
+#include "RenderGraphTimingViewSession.h"
+
+#define LOCTEXT_NAMESPACE "RenderGraphTimingViewExtender"
+
+namespace UE
+{
+namespace RenderGraphInsights
+{
+
+void FRenderGraphTimingViewExtender::OnBeginSession(Insights::ITimingViewSession& InSession)
+{
+	FPerSessionData* PerSessionData = PerSessionDataMap.Find(&InSession);
+	if(PerSessionData == nullptr)
+	{
+		PerSessionData = &PerSessionDataMap.Add(&InSession);
+		PerSessionData->SharedData = MakeUnique<FRenderGraphTimingViewSession>();
+
+		PerSessionData->SharedData->OnBeginSession(InSession);
+	}
+	else
+	{
+		PerSessionData->SharedData->OnBeginSession(InSession);
+	}
+}
+
+void FRenderGraphTimingViewExtender::OnEndSession(Insights::ITimingViewSession& InSession)
+{
+	FPerSessionData* PerSessionData = PerSessionDataMap.Find(&InSession);
+	if(PerSessionData != nullptr)
+	{
+		PerSessionData->SharedData->OnEndSession(InSession);
+	}
+
+	PerSessionDataMap.Remove(&InSession);
+}
+
+void FRenderGraphTimingViewExtender::Tick(Insights::ITimingViewSession& InSession, const TraceServices::IAnalysisSession& InAnalysisSession)
+{
+	FPerSessionData* PerSessionData = PerSessionDataMap.Find(&InSession);
+	if(PerSessionData != nullptr)
+	{
+		PerSessionData->SharedData->Tick(InSession, InAnalysisSession);
+	}
+}
+
+void FRenderGraphTimingViewExtender::ExtendFilterMenu(Insights::ITimingViewSession& InSession, FMenuBuilder& InMenuBuilder)
+{
+	FPerSessionData* PerSessionData = PerSessionDataMap.Find(&InSession);
+	if (PerSessionData != nullptr)
+	{
+		PerSessionData->SharedData->ExtendFilterMenu(InMenuBuilder);
+	}
+}
+
+} //namespace RenderGraphInsights
+} //namespace UE
+
+#undef LOCTEXT_NAMESPACE
\ No newline at end of file
diff --git a/Engine/Plugins/RenderGraphInsights/Source/RenderGraphInsights/Private/RenderGraphTimingViewExtender.h b/Engine/Plugins/RenderGraphInsights/Source/RenderGraphInsights/Private/RenderGraphTimingViewExtender.h
new file mode 100644
--- /dev/null
+++ b/Engine/Plugins/RenderGraphInsights/Source/RenderGraphInsights/Private/RenderGraphTimingViewExtender.h
@@ -0,0 +1,39 @@
+// Copyright Epic Games, Inc. All Rights Reserved.
+
+#pragma once
+
+#include "Insights/ITimingViewExtender.h"
+#include "RenderGraphTimingViewSession.h"
+
+namespace Insights { class ITimingViewSession; }
+namespace TraceServices { class IAnalysisSession; }
+
+class FMenuBuilder;
+
+namespace UE
+{
+namespace RenderGraphInsights
+{
+
+class FRenderGraphTimingViewExtender : public Insights::ITimingViewExtender
+{
+public:
+	//~ Begin Insights::ITimingViewExtender interface
+	virtual void OnBeginSession(Insights::ITimingViewSession& InSession) override;
+	virtual void OnEndSession(Insights::ITimingViewSession& InSession) override;
+	virtual void Tick(Insights::ITimingViewSession& InSession, const TraceServices::IAnalysisSession& InAnalysisSession) override;
+	virtual void ExtendFilterMenu(Insights::ITimingViewSession& InSession, FMenuBuilder& InMenuBuilder) override;
+	//~ End Insights::ITimingViewExtender interface
+
+private:
+	struct FPerSessionData
+	{
+		TUniquePtr<FRenderGraphTimingViewSession> SharedData;
+	};
+
+	// The data we host per-session
+	TMap<Insights::ITimingViewSession*, FPerSessionData> PerSessionDataMap;
+};
+
+} //namespace SlateInsights
+} //namespace UE
diff --git a/Engine/Plugins/RenderGraphInsights/Source/RenderGraphInsights/Private/RenderGraphTimingViewSession.cpp b/Engine/Plugins/RenderGraphInsights/Source/RenderGraphInsights/Private/RenderGraphTimingViewSession.cpp
new file mode 100644
--- /dev/null
+++ b/Engine/Plugins/RenderGraphInsights/Source/RenderGraphInsights/Private/RenderGraphTimingViewSession.cpp
@@ -0,0 +1,94 @@
+// Copyright Epic Games, Inc. All Rights Reserved.
+
+#include "RenderGraphTimingViewSession.h"
+
+#include "Insights/ITimingViewSession.h"
+#include "Modules/ModuleManager.h"
+#include "TraceServices/Model/AnalysisSession.h"
+
+#include "RenderGraphTrack.h"
+#include "RenderGraphInsightsModule.h"
+#include "RenderGraphProvider.h"
+#include "RenderGraphTimingViewExtender.h"
+
+#include "Framework/MultiBox/MultiBoxBuilder.h"
+#include "Widgets/Docking/SDockTab.h"
+
+#define LOCTEXT_NAMESPACE "RenderGraphTimingViewSession"
+
+namespace UE
+{
+namespace RenderGraphInsights
+{
+
+void FRenderGraphTimingViewSession::OnBeginSession(Insights::ITimingViewSession& InTimingViewSession)
+{
+	TimingViewSession = &InTimingViewSession;
+
+	Track.Reset();
+}
+
+void FRenderGraphTimingViewSession::OnEndSession(Insights::ITimingViewSession& InTimingViewSession)
+{
+	Track.Reset();
+	TimingViewSession = nullptr;
+}
+
+void FRenderGraphTimingViewSession::Tick(Insights::ITimingViewSession& InTimingViewSession, const TraceServices::IAnalysisSession& InAnalysisSession)
+{
+	AnalysisSession = &InAnalysisSession;
+
+	const FRenderGraphProvider* RenderGraphProvider = InAnalysisSession.ReadProvider<FRenderGraphProvider>(FRenderGraphProvider::ProviderName);
+
+	if (RenderGraphProvider)
+	{
+		TraceServices::FAnalysisSessionReadScope SessionReadScope(GetAnalysisSession());
+
+		const FRenderGraphProvider::TGraphTimeline& GraphTimeline = RenderGraphProvider->GetGraphTimeline();
+
+		if (GraphTimeline.GetEventCount() && !Track.IsValid())
+		{
+			Track = MakeShared<FRenderGraphTrack>(*this);
+			Track->SetVisibilityFlag(bTrackVisible);
+			Track->SetOrder(MAX_int32);
+
+			InTimingViewSession.AddScrollableTrack(Track);
+			InTimingViewSession.InvalidateScrollableTracksOrder();
+		}
+	}
+}
+
+void FRenderGraphTimingViewSession::ExtendFilterMenu(FMenuBuilder& InMenuBuilder)
+{
+	InMenuBuilder.BeginSection("RenderGraphTracks", LOCTEXT("RenderGraphHeader", "RDG"));
+	{
+		InMenuBuilder.AddMenuEntry(
+			LOCTEXT("RenderGraphTimingTracks", "RDG Tracks"),
+			LOCTEXT("RenderGraphTimingTracks_Tooltip", "Show/hide the RDG track"),
+			FSlateIcon(),
+			FUIAction(
+				FExecuteAction::CreateRaw(this, &FRenderGraphTimingViewSession::ToggleRenderGraphTrack),
+				FCanExecuteAction(),
+				FIsActionChecked::CreateLambda([this]() { return bTrackVisible; })
+				),
+			NAME_None,
+			EUserInterfaceActionType::ToggleButton
+		);
+	}
+	InMenuBuilder.EndSection();
+}
+
+void FRenderGraphTimingViewSession::ToggleRenderGraphTrack()
+{
+	bTrackVisible = !bTrackVisible;
+
+	if (Track)
+	{
+		Track->SetVisibilityFlag(bTrackVisible);
+	}
+}
+
+} //namespace RenderGraphInsights
+} //namespace UE
+
+#undef LOCTEXT_NAMESPACE
diff --git a/Engine/Plugins/RenderGraphInsights/Source/RenderGraphInsights/Private/RenderGraphTimingViewSession.h b/Engine/Plugins/RenderGraphInsights/Source/RenderGraphInsights/Private/RenderGraphTimingViewSession.h
new file mode 100644
--- /dev/null
+++ b/Engine/Plugins/RenderGraphInsights/Source/RenderGraphInsights/Private/RenderGraphTimingViewSession.h
@@ -0,0 +1,47 @@
+// Copyright Epic Games, Inc. All Rights Reserved.
+
+#pragma once
+
+#include "Insights/ViewModels/TimingEventsTrack.h"
+#include "TraceServices/Model/Frames.h"
+
+namespace TraceServices { class IAnalysisSession; }
+namespace Insights { class ITimingViewSession; }
+namespace Insights { enum class ETimeChangedFlags : int32; }
+class FMenuBuilder;
+class SDockTab;
+
+namespace UE
+{
+namespace RenderGraphInsights
+{
+class FRenderGraphTrack;
+
+class FRenderGraphTimingViewSession
+{
+public:
+	FRenderGraphTimingViewSession() = default;
+
+	void OnBeginSession(Insights::ITimingViewSession& InTimingViewSession);
+	void OnEndSession(Insights::ITimingViewSession& InTimingViewSession);
+	void Tick(Insights::ITimingViewSession& InTimingViewSession, const TraceServices::IAnalysisSession& InAnalysisSession);
+	void ExtendFilterMenu(FMenuBuilder& InMenuBuilder);
+
+	/** Get the last cached analysis session */
+	const TraceServices::IAnalysisSession& GetAnalysisSession() const { return *AnalysisSession; }
+
+	/** Check whether the analysis session is valid */
+	bool IsAnalysisSessionValid() const { return AnalysisSession != nullptr; }
+
+	/** Show/Hide the RenderGraph track */
+	void ToggleRenderGraphTrack();
+
+private:
+	const TraceServices::IAnalysisSession* AnalysisSession = nullptr;
+	Insights::ITimingViewSession* TimingViewSession;
+	TSharedPtr<FRenderGraphTrack> Track;
+	bool bTrackVisible = true;
+};
+
+} //namespace RenderGraphInsights
+} //namespace UE
diff --git a/Engine/Plugins/RenderGraphInsights/Source/RenderGraphInsights/Private/RenderGraphTraceModule.cpp b/Engine/Plugins/RenderGraphInsights/Source/RenderGraphInsights/Private/RenderGraphTraceModule.cpp
new file mode 100644
--- /dev/null
+++ b/Engine/Plugins/RenderGraphInsights/Source/RenderGraphInsights/Private/RenderGraphTraceModule.cpp
@@ -0,0 +1,39 @@
+// Copyright Epic Games, Inc. All Rights Reserved.
+
+#include "RenderGraphTraceModule.h"
+#include "RenderGraphProvider.h"
+#include "RenderGraphAnalyzer.h"
+
+namespace UE
+{
+namespace RenderGraphInsights
+{
+
+FName FRenderGraphTraceModule::ModuleName("TraceModule_RenderGraph");
+
+void FRenderGraphTraceModule::GetModuleInfo(TraceServices::FModuleInfo& OutModuleInfo)
+{
+	OutModuleInfo.Name = ModuleName;
+	OutModuleInfo.DisplayName = TEXT("RDG");
+}
+
+void FRenderGraphTraceModule::OnAnalysisBegin(TraceServices::IAnalysisSession& InSession)
+{
+	FRenderGraphProvider* RenderGraphProvider = new FRenderGraphProvider(InSession);
+	InSession.AddProvider(FRenderGraphProvider::ProviderName, RenderGraphProvider);
+
+	InSession.AddAnalyzer(new FRenderGraphAnalyzer(InSession, *RenderGraphProvider));
+}
+
+void FRenderGraphTraceModule::GetLoggers(TArray<const TCHAR *>& OutLoggers)
+{
+	OutLoggers.Add(TEXT("RDG"));
+}
+
+void FRenderGraphTraceModule::GenerateReports(const TraceServices::IAnalysisSession& Session, const TCHAR* CmdLine, const TCHAR* OutputDirectory)
+{
+
+}
+
+} //namespace RenderGraphInsights
+} //namespace UE
diff --git a/Engine/Plugins/RenderGraphInsights/Source/RenderGraphInsights/Private/RenderGraphTraceModule.h b/Engine/Plugins/RenderGraphInsights/Source/RenderGraphInsights/Private/RenderGraphTraceModule.h
new file mode 100644
--- /dev/null
+++ b/Engine/Plugins/RenderGraphInsights/Source/RenderGraphInsights/Private/RenderGraphTraceModule.h
@@ -0,0 +1,29 @@
+// Copyright Epic Games, Inc. All Rights Reserved.
+
+#pragma once
+
+#include "CoreMinimal.h"
+#include "TraceServices/ModuleService.h"
+
+namespace UE
+{
+namespace RenderGraphInsights
+{
+
+class FRenderGraphTraceModule : public TraceServices::IModule
+{
+public:
+	//~ Begin TraceServices::IModule interface
+	virtual void GetModuleInfo(TraceServices::FModuleInfo& OutModuleInfo) override;
+	virtual void OnAnalysisBegin(TraceServices::IAnalysisSession& Session) override;
+	virtual void GetLoggers(TArray<const TCHAR *>& OutLoggers) override;
+	virtual void GenerateReports(const TraceServices::IAnalysisSession& Session, const TCHAR* CmdLine, const TCHAR* OutputDirectory) override;
+	virtual const TCHAR* GetCommandLineArgument() override { return TEXT("rdgtrace"); }
+	//~ End TraceServices::IModule interface
+
+private:
+	static FName ModuleName;
+};
+
+} //namespace RenderGraphInsights
+} //namespace UE
diff --git a/Engine/Plugins/RenderGraphInsights/Source/RenderGraphInsights/Private/RenderGraphTrack.cpp b/Engine/Plugins/RenderGraphInsights/Source/RenderGraphInsights/Private/RenderGraphTrack.cpp
new file mode 100644
--- /dev/null
+++ b/Engine/Plugins/RenderGraphInsights/Source/RenderGraphInsights/Private/RenderGraphTrack.cpp
@@ -0,0 +1,1126 @@
+// Copyright Epic Games, Inc. All Rights Reserved.
+
+#include "RenderGraphTrack.h"
+
+#include "Modules/ModuleManager.h"
+#include "RenderGraphProvider.h"
+#include "RenderGraphTimingViewSession.h"
+
+#include "Application/SlateApplicationBase.h"
+#include "Fonts/FontMeasure.h"
+#include "Framework/MultiBox/MultiBoxBuilder.h"
+
+#include "Insights/Common/PaintUtils.h"
+#include "Insights/ViewModels/GraphSeries.h"
+#include "Insights/ViewModels/GraphTrackBuilder.h"
+#include "Insights/ViewModels/ITimingViewDrawHelper.h"
+#include "Insights/ViewModels/TimingTrackViewport.h"
+#include "Insights/ViewModels/TimingEvent.h"
+#include "Insights/ViewModels/TimingEventSearch.h"
+#include "Insights/ViewModels/TooltipDrawState.h"
+
+#include "TraceServices/Model/Frames.h"
+
+#define LOCTEXT_NAMESPACE "RenderGraphTrack"
+
+namespace UE
+{
+namespace RenderGraphInsights
+{
+
+class FEvent : public FTimingEvent
+{
+	INSIGHTS_DECLARE_RTTI(FEvent, FTimingEvent)
+public:
+	FEvent(const TSharedRef<const FBaseTimingTrack> InTrack, const FVisibleItem& InItem)
+		: FTimingEvent(InTrack, InItem.Packet->StartTime, InItem.Packet->EndTime, InItem.Depth)
+	{}
+
+	virtual const FVisibleItem& GetItem() const = 0;
+	virtual const FPacket& GetPacket() const = 0;
+};
+
+template <typename ItemType>
+class TEventHelper : public FEvent
+{
+public:
+	TEventHelper(const TSharedRef<const FBaseTimingTrack>& InTrack, const ItemType& InItem)
+		: FEvent(InTrack, InItem)
+		, Item(InItem)
+	{
+		Item.Index = kInvalidVisibleIndex;
+		Item.StartTime = GetPacket().StartTime;
+		Item.EndTime = GetPacket().EndTime;
+	}
+
+	const ItemType& GetItem() const override
+	{
+		return Item;
+	}
+
+	const typename ItemType::PacketType& GetPacket() const
+	{
+		return Item.GetPacket();
+	}
+
+private:
+	ItemType Item;
+};
+
+class FVisibleScopeEvent final : public TEventHelper<FVisibleScope>
+{
+	INSIGHTS_DECLARE_RTTI(FVisibleScopeEvent, FEvent)
+public:
+	FVisibleScopeEvent(const TSharedRef<const FBaseTimingTrack>& InTrack, const FVisibleScope& InItem)
+		: TEventHelper<FVisibleScope>(InTrack, InItem)
+	{}
+};
+
+class FVisiblePassEvent final : public TEventHelper<FVisiblePass>
+{
+	INSIGHTS_DECLARE_RTTI(FVisiblePassEvent, FEvent)
+public:
+	FVisiblePassEvent(const TSharedRef<const FBaseTimingTrack>& InTrack, const FVisiblePass& InItem)
+		: TEventHelper<FVisiblePass>(InTrack, InItem)
+	{}
+};
+
+class FVisibleTextureEvent final : public TEventHelper<FVisibleTexture>
+{
+	INSIGHTS_DECLARE_RTTI(FVisibleTextureEvent, FEvent)
+public:
+	FVisibleTextureEvent(const TSharedRef<const FBaseTimingTrack>& InTrack, const FVisibleTexture& InItem)
+		: TEventHelper<FVisibleTexture>(InTrack, InItem)
+	{}
+};
+
+class FVisibleBufferEvent final : public TEventHelper<FVisibleBuffer>
+{
+	INSIGHTS_DECLARE_RTTI(FVisibleBufferEvent, FEvent)
+public:
+	FVisibleBufferEvent(const TSharedRef<const FBaseTimingTrack>& InTrack, const FVisibleBuffer& InItem)
+		: TEventHelper<FVisibleBuffer>(InTrack, InItem)
+	{}
+};
+
+class FPacketFilter final : public ITimingEventFilter
+{
+	INSIGHTS_DECLARE_RTTI(FPacketFilter, ITimingEventFilter)
+public:
+	FPacketFilter(const TSharedPtr<const FEvent>& InEvent)
+		: Event(InEvent)
+		, Packet(&InEvent->GetPacket())
+		, Graph(Packet->Graph)
+	{}
+	virtual ~FPacketFilter() {}
+
+	bool FilterPacket(const FPacket& Packet) const;
+	bool FilterPacketExact(const FPacket& Packet) const;
+
+	const FPacket& GetPacket() const
+	{
+		return *Packet;
+	}
+
+	const FGraphPacket& GetGraph() const
+	{
+		return *Graph;
+	}
+
+private:
+	//! ITimingEventFilter
+	bool FilterTrack(const FBaseTimingTrack& InTrack) const override { return true; }
+	bool FilterEvent(const ITimingEvent& InEvent) const override { return true; }
+	bool FilterEvent(double InEventStartTime, double InEventEndTime, uint32 InEventDepth, const TCHAR* InEventName, uint64 InEventType = 0, uint32 InEventColor = 0) const override { return true; }
+	uint32 GetChangeNumber() const override { return 0; }
+	//! ITimingEventFilter
+
+	TSharedPtr<const FEvent> Event;
+	const FPacket* Packet{};
+	const FGraphPacket* Graph{};
+};
+
+INSIGHTS_IMPLEMENT_RTTI(FRenderGraphTrack)
+INSIGHTS_IMPLEMENT_RTTI(FPacketFilter)
+INSIGHTS_IMPLEMENT_RTTI(FEvent)
+INSIGHTS_IMPLEMENT_RTTI(FVisibleScopeEvent)
+INSIGHTS_IMPLEMENT_RTTI(FVisiblePassEvent)
+INSIGHTS_IMPLEMENT_RTTI(FVisibleTextureEvent)
+INSIGHTS_IMPLEMENT_RTTI(FVisibleBufferEvent)
+
+template <typename VisibleItemType>
+static VisibleItemType& AddVisibleItem(TArray<VisibleItemType>& VisibleItems, const VisibleItemType& InVisibleItem)
+{
+	const uint32 VisibleIndex = VisibleItems.Num();
+
+	VisibleItemType& VisibleItem = VisibleItems.Emplace_GetRef(InVisibleItem);
+	check(VisibleItem.Packet->VisibleIndex == kInvalidVisibleIndex);
+	VisibleItem.Packet->VisibleIndex = VisibleIndex;
+	VisibleItem.Index = VisibleIndex;
+	return VisibleItem;
+}
+
+template <typename VisibleItemType>
+static void ResetVisibleItemArray(TArray<VisibleItemType>& VisibleItems)
+{
+	for (VisibleItemType& VisibleItem : VisibleItems)
+	{
+		check(VisibleItem.Packet->VisibleIndex == VisibleItem.Index);
+		VisibleItem.Packet->VisibleIndex = kInvalidVisibleIndex;
+	}
+	VisibleItems.Empty();
+}
+
+static bool Intersects(const FPassIntervalPacket& A, const FPassPacket& B)
+{
+	return A.FirstPass <= B.Handle && A.LastPass >= B.Handle;
+}
+
+static bool Intersects(const FPassIntervalPacket& A, const FPassIntervalPacket& B)
+{
+	return !(B.LastPass < A.FirstPass || A.LastPass < B.FirstPass);
+}
+
+bool FPacketFilter::FilterPacketExact(const FPacket& PacketToFilter) const
+{
+	return Packet == &PacketToFilter;
+}
+
+bool FPacketFilter::FilterPacket(const FPacket& PacketToFilter) const
+{
+	if (FilterPacketExact(PacketToFilter))
+	{
+		return true;
+	}
+
+	if (PacketToFilter.Graph != Graph)
+	{
+		return false;
+	}
+
+	if (Packet->Is<FScopePacket>())
+	{
+		const FScopePacket& Scope = Packet->As<FScopePacket>();
+
+		if (PacketToFilter.Is<FPassIntervalPacket>())
+		{
+			const FPassIntervalPacket& PassIntervalToFilter = PacketToFilter.As<FPassIntervalPacket>();
+			return Intersects(Scope, PassIntervalToFilter);
+		}
+		else if (PacketToFilter.Is<FPassPacket>())
+		{
+			const FPassPacket& PassToFilter = PacketToFilter.As<FPassPacket>();
+			return Intersects(Scope, PassToFilter);
+		}
+	}
+
+	return false;
+}
+
+bool FRenderGraphTrack::FilterPacket(const FPacket& Packet) const
+{
+	return PacketFilter ? PacketFilter->FilterPacket(Packet) : false;
+}
+
+bool FRenderGraphTrack::FilterPacketExact(const FPacket& Packet) const
+{
+	return PacketFilter ? PacketFilter->FilterPacketExact(Packet) : false;
+}
+
+void FVisibleGraph::Reset()
+{
+	ResetVisibleItemArray(Scopes);
+	ResetVisibleItemArray(Passes);
+	ResetVisibleItemArray(Textures);
+	ResetVisibleItemArray(Buffers);
+}
+
+void FVisibleGraph::AddScope(const FVisibleScope& VisibleScope)
+{
+	AddVisibleItem(Scopes, VisibleScope);
+}
+
+void FVisibleGraph::AddPass(const FVisiblePass& InVisiblePass)
+{
+	const FVisiblePass& VisiblePass = AddVisibleItem(Passes, InVisiblePass);
+
+	if (EnumHasAnyFlags(VisiblePass.GetPacket().Flags, ERDGPassFlags::AsyncCompute))
+	{
+		AsyncComputePasses.Add(VisiblePass.Index);
+	}
+}
+
+void FVisibleGraph::AddTexture(const FVisibleTexture& VisibleTexture)
+{
+	AddVisibleItem(Textures, VisibleTexture);
+}
+
+void FVisibleGraph::AddBuffer(const FVisibleBuffer& VisibleBuffer)
+{
+	AddVisibleItem(Buffers, VisibleBuffer);
+}
+
+const FVisibleItem* FVisibleGraph::FindItem(float InPosX, float InPosY) const
+{
+	for (const FVisibleTexture& Texture : Textures)
+	{
+		if (Texture.Intersects(InPosX, InPosY))
+		{
+			return &Texture;
+		}
+	}
+
+	for (const FVisibleBuffer& Buffer : Buffers)
+	{
+		if (Buffer.Intersects(InPosX, InPosY))
+		{
+			return &Buffer;
+		}
+	}
+
+	for (const FVisibleScope& Scope : Scopes)
+	{
+		if (Scope.Intersects(InPosX, InPosY))
+		{
+			return &Scope;
+		}
+	}
+
+	for (const FVisiblePass& Pass : Passes)
+	{
+		if (Pass.Intersects(InPosX, InPosY))
+		{
+			return &Pass;
+		}
+	}
+
+	return nullptr;
+}
+
+FRenderGraphTrack::FRenderGraphTrack(const FRenderGraphTimingViewSession& InSharedData)
+	: Super(LOCTEXT("TrackNameFormat", "RDG").ToString())
+	, SharedData(InSharedData)
+{}
+
+static const uint32 kBuilderColor = 0xffa0a0a0;
+static const uint32 kRasterPassColor = 0xff7F2D2D;
+static const uint32 kComputePassColor = 0xff2D9F9F;
+static const uint32 kNoParameterPassColor = 0xff4D4D4D;
+static const uint32 kAsyncComputePassColor = 0xff2D7f2d;
+static const uint32 kTextureColor = 0xff89cff0;
+static const uint32 kBufferColor = 0xff66D066;
+
+static const float kMinGraphPixels = 5.0f;
+static const float kMinPassMarginPixels = 5.0f;
+
+static uint32 GetPassColor(const FPassPacket& Packet)
+{
+	const ERDGPassFlags Flags = Packet.Flags;
+
+	uint32 Color = 0;
+
+	const bool bNoParameterPass = (!Packet.Buffers.Num() && !Packet.Textures.Num());
+
+	if (bNoParameterPass && !Packet.bCulled)
+	{
+		Color = kNoParameterPassColor;
+	}
+	else if (EnumHasAnyFlags(Flags, ERDGPassFlags::AsyncCompute))
+	{
+		Color = kAsyncComputePassColor;
+	}
+	else if (EnumHasAnyFlags(Flags, ERDGPassFlags::Raster))
+	{
+		Color = kRasterPassColor;
+	}
+	else
+	{
+		Color = kComputePassColor;
+	}
+
+	if (Packet.bCulled)
+	{
+		Color &= 0x00FFFFFF;
+		Color |= 0x40000000;
+	}
+
+	return Color;
+}
+
+void FRenderGraphTrack::Update(const ITimingTrackUpdateContext& Context)
+{
+	FTimingEventsTrack::Update(Context);
+
+	MousePosition = Context.GetMousePosition();
+	MousePosition.Y -= GetPosY();
+
+	SelectedPass = nullptr;
+
+	const TSharedPtr<const ITimingEvent> SelectedEventPtr = Context.GetSelectedEvent();
+	if (SelectedEventPtr.IsValid() && SelectedEventPtr->Is<FEvent>())
+	{
+		if (SelectedEventPtr->Is<FVisiblePassEvent>())
+		{
+			SelectedPass = &SelectedEventPtr->As<FVisiblePassEvent>().GetPacket();
+		}
+
+		const FEvent& SelectedEvent = SelectedEventPtr->As<FEvent>();
+
+		const FTimingTrackViewport& Viewport = Context.GetViewport();
+
+		const FVector2D ViewportPos(Viewport.GetWidth(), Viewport.GetHeight());
+
+		InitTooltip(SelectedTooltipState, SelectedEvent);
+		SelectedTooltipState.SetDesiredOpacity(0.75f);
+		SelectedTooltipState.SetPosition(ViewportPos, 0.0f, Viewport.GetWidth(), 0.0f, Viewport.GetHeight());
+		SelectedTooltipState.Update();
+	}
+	else
+	{
+		SelectedTooltipState.Reset();
+	}
+}
+
+void FRenderGraphTrack::Draw(const ITimingTrackDrawContext& Context) const
+{
+	FTimingEventsTrack::Draw(Context);
+
+	const FEvent* HoveredEvent = nullptr;
+	if (TSharedPtr<const ITimingEvent> HoveredEventBase = Context.GetHoveredEvent())
+	{
+		if (HoveredEventBase->Is<FEvent>())
+		{
+			HoveredEvent = &HoveredEventBase->As<FEvent>();
+		}
+	}
+
+	const int32 LineLayerId = Context.GetHelper().GetFirstLayerId() - 1;
+	const FTimingTrackViewport& Viewport = Context.GetViewport();
+	const FDrawContext& DrawContext = Context.GetDrawContext();
+	const ITimingViewDrawHelper& DrawHelper = Context.GetHelper();
+	const FLinearColor EdgeColor = DrawHelper.GetEdgeColor();
+
+	const float TrackY = GetPosY();
+	const float TrackH = GetHeight();
+	const float EventH = Viewport.GetLayout().EventH;
+	const float ViewportWidth = Viewport.GetWidth();
+
+	const auto DrawClampedBox = [&](int32 LayerId, float MinX, float MinY, float W, float H, const FLinearColor& Color)
+	{
+		const float Guardband = 1024.0f;
+		float MaxX = MinX + W;
+
+		if (MinX > ViewportWidth || MaxX < 0.0f)
+		{
+			return;
+		}
+
+		MinX = FMath::Max(MinX, -Guardband);
+		MaxX = FMath::Min(MaxX, ViewportWidth + Guardband);
+		W = MaxX - MinX;
+
+		DrawContext.DrawBox(LayerId, MinX, MinY, W, H, DrawHelper.GetWhiteBrush(), Color);
+	};
+
+	const auto DrawClampedSpline = [&](int32 SplineLayerId, FSpline Spline)
+	{
+		const float Guardband = 1024.0f;
+		float MinX = Spline.Start.X;
+		float MaxX = Spline.Start.X + Spline.End.X;
+
+		if (MinX > ViewportWidth || MaxX < 0.0f)
+		{
+			return;
+		}
+
+		MinX = FMath::Max(MinX, -Guardband);
+		MaxX = FMath::Min(MaxX, ViewportWidth + Guardband);
+
+		const float Width = MaxX - MinX;
+
+		const float MinY = TrackY + Spline.Start.Y;
+
+		Spline.Start.X = MinX;
+		Spline.End.X = Width;
+		Spline.StartDir.X = FMath::Min(Guardband, Spline.StartDir.X);
+		Spline.EndDir.X = FMath::Min(Guardband, Spline.EndDir.X);
+
+		DrawContext.DrawSpline(SplineLayerId, MinX, MinY, FVector2D::ZeroVector, Spline.StartDir, Spline.End, Spline.EndDir, Spline.Thickness, Spline.Tint);
+	};
+
+	if (TrackH > 0.0f)
+	{
+		const FTimingViewLayout& Layout = Viewport.GetLayout();
+
+		for (const FVisibleGraph& VisibleGraph : VisibleGraphs)
+		{
+			const FGraphPacket& Graph = VisibleGraph.GetPacket();
+			const uint32 PassCount = Graph.PassCount;
+			const float LineStrideMin = 5.0f;
+			const float GraphLineStride = VisibleGraph.Max.X - VisibleGraph.Min.X;
+
+			if (GraphLineStride >= LineStrideMin)
+			{
+				DrawClampedBox(LineLayerId, VisibleGraph.Min.X, TrackY, 1.0f, TrackH, EdgeColor);
+				DrawClampedBox(LineLayerId, VisibleGraph.Max.X, TrackY, 1.0f, TrackH, EdgeColor);
+			}
+
+			const float PassLineStride = GraphLineStride / float(PassCount);
+
+			FVector2D RenderPassMergeMin = FVector2D::ZeroVector;
+
+			for (uint32 PassIndex = 0; PassIndex < PassCount; ++PassIndex)
+			{
+				const FPassPacket& Pass = Graph.Passes[PassIndex];
+				const FVisiblePass& VisiblePass = VisibleGraph.GetVisiblePass(Pass);
+
+				if (!Pass.bSkipRenderPassBegin && Pass.bSkipRenderPassEnd)
+				{
+					RenderPassMergeMin = VisiblePass.Min;
+				}
+
+				if (Pass.bSkipRenderPassBegin && !Pass.bSkipRenderPassEnd)
+				{
+					const FVector2D RenderPassMergeMax = VisiblePass.Max;
+					const float RenderPassMarginY = 3.0f;
+					const float W = RenderPassMergeMax.X - RenderPassMergeMin.X;
+					const float H = (RenderPassMergeMax.Y - RenderPassMergeMin.Y) * 0.25f;
+					const float X = RenderPassMergeMin.X;
+					const float Y = RenderPassMergeMin.Y - H - RenderPassMarginY;
+
+					DrawClampedBox(LineLayerId, X, TrackY + Y, W, H, FLinearColor(0.8f, 0.2f, 0.2f, 0.75f));
+				}
+
+				const float X = Viewport.TimeToSlateUnitsRounded(Pass.StartTime);
+				const float Y = TrackY + VisiblePass.Max.Y;
+				const float H = TrackH - VisiblePass.Max.Y;
+
+				const bool bHoveredPass = HoveredEvent && &HoveredEvent->GetPacket() == &Pass;
+				const bool bSelectedPass = SelectedPass == &Pass;
+				const bool bFilteredPass = FilterPacket(Pass);
+
+				if (bHoveredPass || bFilteredPass || bSelectedPass)
+				{
+					const float W = PassLineStride + 1;
+
+					const FLinearColor EdgeColorTranslucent = FLinearColor(EdgeColor.R, EdgeColor.G, EdgeColor.B, bFilteredPass ? 1.0f : 0.5f);
+
+					DrawClampedBox(LineLayerId, X, Y, W, H, EdgeColorTranslucent);
+				}
+				else if (PassLineStride >= LineStrideMin && PassIndex != 0)
+				{
+					const float W = 1.0f;
+
+					DrawClampedBox(LineLayerId, X, Y, W, H, EdgeColor);
+				}
+			}
+		}
+	}
+
+	const int32 SplineOverLayerId = LineLayerId + Context.GetHelper().GetNumLayerIds();
+	const int32 SplineUnderLineLayerId = LineLayerId + 1;
+
+	for (const FSpline& Spline : Splines)
+	{
+		DrawClampedSpline(SplineOverLayerId, Spline);
+	}
+
+	const bool bFilterActive = Context.GetEventFilter() != nullptr;
+
+	for (const FVisibleGraph& VisibleGraph : VisibleGraphs)
+	{
+		const FGraphPacket& Graph = VisibleGraph.GetPacket();
+
+		for (uint32 VisibleIndex : VisibleGraph.AsyncComputePasses)
+		{
+			const FVisiblePass& AsyncComputeVisiblePass = VisibleGraph.Passes[VisibleIndex];
+			const FPassPacket& AsyncComputePass = AsyncComputeVisiblePass.GetPacket();
+
+			float TintAlpha = 0.25f;
+
+			if (!bFilterActive || FilterPacket(AsyncComputePass))
+			{
+				TintAlpha = 0.75f;
+			}
+
+			const float StartT = 0.2f;
+			const float EndT = 1.0f - StartT;
+			const float SplineDir = 20.0f;
+
+			const FRDGPassHandle GraphicsForkHandle = AsyncComputePass.GraphicsForkPass;
+
+			if (GraphicsForkHandle.IsValid() && AsyncComputePass.bAsyncComputeBegin)
+			{
+				const FPassPacket& GraphicsForkPass = *Graph.GetPass(AsyncComputePass.GraphicsForkPass);
+				const FVisiblePass& GraphicsForkVisiblePass = VisibleGraph.GetVisiblePass(GraphicsForkPass);
+
+				const float X = FMath::Lerp(GraphicsForkVisiblePass.Min.X, GraphicsForkVisiblePass.Max.X, EndT);
+				const float Y = GraphicsForkVisiblePass.Max.Y;
+				const float LocalEndX = FMath::Lerp(AsyncComputeVisiblePass.Min.X, AsyncComputeVisiblePass.Max.X, StartT) - X;
+				const float LocalEndY = AsyncComputeVisiblePass.Min.Y - Y;
+
+				FSpline Spline;
+				Spline.Start.X = X;
+				Spline.Start.Y = Y;
+				Spline.StartDir = FVector2D(0, SplineDir);
+				Spline.End = FVector2D(LocalEndX, LocalEndY);
+				Spline.EndDir = FVector2D(0, SplineDir);
+				Spline.Thickness = 2.0f;
+				Spline.Tint = FLinearColor(0.4f, 1.0f, 0.4f, TintAlpha);
+				DrawClampedSpline(SplineUnderLineLayerId + 1, Spline);
+			}
+
+			const FRDGPassHandle GraphicsJoinHandle = AsyncComputePass.GraphicsJoinPass;
+
+			if (GraphicsJoinHandle.IsValid() && AsyncComputePass.bAsyncComputeEnd)
+			{
+				const FPassPacket& GraphicsJoinPass = *Graph.GetPass(AsyncComputePass.GraphicsJoinPass);
+				const FVisiblePass& GraphicsJoinVisiblePass = VisibleGraph.GetVisiblePass(GraphicsJoinPass);
+
+				const float X = FMath::Lerp(AsyncComputeVisiblePass.Min.X, AsyncComputeVisiblePass.Max.X, EndT);
+				const float Y = AsyncComputeVisiblePass.Min.Y;
+				const float LocalEndX = FMath::Lerp(GraphicsJoinVisiblePass.Min.X, GraphicsJoinVisiblePass.Max.X, StartT) - X;
+				const float LocalEndY = GraphicsJoinVisiblePass.Max.Y - Y;
+
+				FSpline Spline;
+				Spline.Start.X = X;
+				Spline.Start.Y = Y;
+				Spline.StartDir = FVector2D(0, -SplineDir);
+				Spline.End = FVector2D(LocalEndX, LocalEndY);
+				Spline.EndDir = FVector2D(0, -SplineDir);
+				Spline.Thickness = 2.0f;
+				Spline.Tint = FLinearColor(1.0f, 0.4f, 0.4f, TintAlpha);
+				DrawClampedSpline(SplineUnderLineLayerId + 1, Spline);
+			}
+		}
+	}
+}
+
+void FRenderGraphTrack::PostDraw(const ITimingTrackDrawContext& Context) const
+{
+	SelectedTooltipState.Draw(Context.GetDrawContext());
+}
+
+void FRenderGraphTrack::AddEvent(ITimingEventsTrackDrawStateBuilder& Builder, const FVisibleItem& Item)
+{
+	Builder.AddEvent(Item.StartTime, Item.EndTime, Item.Depth, Item.Color, [&Item](float) { return Item.Name; });
+};
+
+void FRenderGraphTrack::BuildDrawState(ITimingEventsTrackDrawStateBuilder& Builder, const ITimingTrackUpdateContext& Context)
+{
+	// Reset draw state; it's built each frame.
+	{
+		Splines.Empty();
+		for (FVisibleGraph& VisibleGraph : VisibleGraphs)
+		{
+			VisibleGraph.Reset();
+		}
+		ResetVisibleItemArray(VisibleGraphs);
+
+		PacketFilter = nullptr;
+		if (TSharedPtr<ITimingEventFilter> Filter = Context.GetEventFilter())
+		{
+			if (Filter->Is<FPacketFilter>())
+			{
+				PacketFilter = &Filter->As<FPacketFilter>();
+			}
+		}
+	}
+
+	if (const FRenderGraphProvider* RenderGraphProvider = SharedData.GetAnalysisSession().ReadProvider<FRenderGraphProvider>(FRenderGraphProvider::ProviderName))
+	{
+		const FTimingTrackViewport& Viewport = Context.GetViewport();
+
+		TraceServices::FAnalysisSessionReadScope SessionReadScope(SharedData.GetAnalysisSession());
+
+		const FRenderGraphProvider::TGraphTimeline& GraphTimeline = RenderGraphProvider->GetGraphTimeline();
+		GraphTimeline.EnumerateEvents(Viewport.GetStartTime(), Viewport.GetEndTime(),
+			[this, &Builder, &Viewport](double GraphStartTime, double GraphEndTime, uint32, const TSharedPtr<FGraphPacket>& Graph)
+		{
+			// We always render the graph event, so add it separately even if the visible graph is culled.
+			Builder.AddEvent(GraphStartTime, GraphEndTime, 0, *Graph->Name, 0, kBuilderColor);
+
+			if (Viewport.GetViewportDXForDuration(GraphEndTime - GraphStartTime) <= kMinGraphPixels)
+			{
+				return TraceServices::EEventEnumerate::Continue;
+			}
+
+			FVisibleGraph& VisibleGraph = AddVisibleGraph(Viewport, *Graph, kBuilderColor);
+
+			double SinglePixelTimeMargin = Viewport.GetDurationForViewportDX(1.0);
+
+			if (Viewport.GetViewportDXForDuration(Graph->NormalizedPassDuration) <= kMinPassMarginPixels)
+			{
+				SinglePixelTimeMargin = 0.0;
+			}
+
+			uint32 DepthOffset = 1;
+
+			for (uint32 ScopeIndex = 0, ScopeCount = Graph->Scopes.Num(); ScopeIndex < ScopeCount; ++ScopeIndex)
+			{
+				const FScopePacket& Scope = Graph->Scopes[ScopeIndex];
+				const double StartTime = Scope.StartTime + SinglePixelTimeMargin;
+				const double EndTime = Scope.EndTime;
+
+				const FVisibleScope VisibleScope(Viewport, Scope, StartTime, EndTime, DepthOffset + Scope.Depth, FTimingEvent::ComputeEventColor(*Scope.Name));
+				AddEvent(Builder, VisibleScope);
+				VisibleGraph.AddScope(VisibleScope);
+			}
+
+			if (Graph->ScopeDepth)
+			{
+				DepthOffset += Graph->ScopeDepth + 1;
+			}
+
+			// +1 for render pass merge bars.
+			DepthOffset += 1;
+
+			bool bAnyAsyncCompute = false;
+
+			for (uint32 PassIndex = 0, PassCount = Graph->Passes.Num(); PassIndex < PassCount; ++PassIndex)
+			{
+				FPassPacket& Pass = Graph->Passes[PassIndex];
+				const double StartTime = Pass.StartTime + SinglePixelTimeMargin;
+				const double EndTime   = Pass.EndTime;
+				const bool bAsyncCompute = EnumHasAnyFlags(Pass.Flags, ERDGPassFlags::AsyncCompute);
+				const uint32 Depth     = DepthOffset + (bAsyncCompute ? 2 : 0);
+				const uint32 Color     = GetPassColor(Pass);
+
+				const FVisiblePass VisiblePass(Viewport, Pass, StartTime, EndTime, Depth, Color);
+				AddEvent(Builder, VisiblePass);
+				VisibleGraph.AddPass(VisiblePass);
+
+				bAnyAsyncCompute |= bAsyncCompute;
+			}
+
+			// Empty space between passes / resources.
+			DepthOffset += bAnyAsyncCompute ? 3 : 1;
+
+			if (ShowTextures())
+			{
+				uint32 TextureDepthOffsetMax = 0;
+
+				for (uint32 TextureIndex = 0, TextureCount = Graph->Textures.Num(); TextureIndex < TextureCount; ++TextureIndex)
+				{
+					FTexturePacket& Texture = Graph->Textures[TextureIndex];
+					if (Texture.bCulled)
+					{
+						continue;
+					}
+
+					const double StartTime = Texture.StartTime + SinglePixelTimeMargin;
+					const double EndTime   = Texture.EndTime;
+					const uint32 Depth     = DepthOffset + Texture.Depth;
+
+					const FVisibleTexture VisibleTexture(Viewport, Texture, StartTime, EndTime, Depth, kTextureColor);
+					AddEvent(Builder, VisibleTexture);
+					VisibleGraph.AddTexture(VisibleTexture);
+
+					TextureDepthOffsetMax = FMath::Max(Depth, TextureDepthOffsetMax);
+				}
+
+				DepthOffset = FMath::Max(DepthOffset, TextureDepthOffsetMax + 1);
+			}
+
+			if (ShowBuffers())
+			{
+				uint32 BufferDepthOffsetMax = 0;
+
+				for (uint32 BufferIndex = 0, BufferCount = Graph->Buffers.Num(); BufferIndex < BufferCount; ++BufferIndex)
+				{
+					FBufferPacket& Buffer = Graph->Buffers[BufferIndex];
+					if (Buffer.bCulled)
+					{
+						continue;
+					}
+
+					const double StartTime = Buffer.StartTime + SinglePixelTimeMargin;
+					const double EndTime   = Buffer.EndTime;
+					const uint32 Depth     = DepthOffset + Buffer.Depth;
+
+					const FVisibleBuffer VisibleBuffer(Viewport, Buffer, StartTime, EndTime, Depth, kBufferColor);
+					AddEvent(Builder, VisibleBuffer);
+					VisibleGraph.AddBuffer(VisibleBuffer);
+
+					BufferDepthOffsetMax = FMath::Max(Depth, BufferDepthOffsetMax);
+				}
+
+				DepthOffset = FMath::Max(DepthOffset, BufferDepthOffsetMax + 1);
+			}
+
+			check(VisibleGraph.Passes.Num() == VisibleGraph.GetPacket().PassCount);
+
+			return TraceServices::EEventEnumerate::Continue;
+		});
+	}
+}
+
+void FRenderGraphTrack::BuildFilteredDrawState(ITimingEventsTrackDrawStateBuilder& Builder, const ITimingTrackUpdateContext& Context)
+{
+	if (!PacketFilter)
+	{
+		return;
+	}
+
+	const FGraphPacket& Graph = PacketFilter->GetGraph();
+	const FVisibleGraph* VisibleGraph = GetVisibleGraph(Graph);
+
+	if (!VisibleGraph)
+	{
+		return;
+	}
+
+	AddEvent(Builder, *VisibleGraph);
+
+	TSet<const FVisibleItem*> VisibleItems;
+
+	for (const FVisibleScope& VisibleScope : VisibleGraph->Scopes)
+	{
+		if (!FilterPacket(VisibleScope.GetPacket()))
+		{
+			continue;
+		}
+
+		VisibleItems.Add(&VisibleScope);
+	}
+
+	for (const FVisiblePass& VisiblePass : VisibleGraph->Passes)
+	{
+		const FPassPacket& Pass = VisiblePass.GetPacket();
+
+		if (!FilterPacket(Pass))
+		{
+			continue;
+		}
+
+		VisibleItems.Add(&VisiblePass);
+
+		if (ShowTextures())
+		{
+			for (FRDGTextureHandle TextureHandle : Pass.Textures)
+			{
+				const FTexturePacket& Texture = *Graph.GetTexture(TextureHandle);
+				const FVisibleTexture& VisibleTexture = VisibleGraph->GetVisibleTexture(Texture);
+				VisibleItems.Add(&VisibleTexture);
+			}
+		}
+
+		if (ShowBuffers())
+		{
+			for (FRDGBufferHandle BufferHandle : Pass.Buffers)
+			{
+				const FBufferPacket& Buffer = *Graph.GetBuffer(BufferHandle);
+				const FVisibleBuffer& VisibleBuffer = VisibleGraph->GetVisibleBuffer(Buffer);
+				VisibleItems.Add(&VisibleBuffer);
+			}
+		}
+
+		for (const FVisibleScope& VisibleScope : VisibleGraph->Scopes)
+		{
+			const FScopePacket& Scope = VisibleScope.GetPacket();
+			if (Intersects(Scope, Pass))
+			{
+				VisibleItems.Add(&VisibleScope);
+			}
+		}
+
+		const auto AddFencePassEvent = [&](FRDGPassHandle InFencePassHandle)
+		{
+			const FPassPacket& GraphicsPass = *Graph.GetPass(InFencePassHandle);
+			const FVisiblePass& GraphicsVisiblePass = VisibleGraph->GetVisiblePass(GraphicsPass);
+			VisibleItems.Add(&GraphicsVisiblePass);
+		};
+
+		if (Pass.bAsyncComputeBegin)
+		{
+			AddFencePassEvent(Pass.GraphicsForkPass);
+		}
+
+		if (Pass.bAsyncComputeEnd)
+		{
+			AddFencePassEvent(Pass.GraphicsJoinPass);
+		}
+	}
+
+	const auto AddResourcePassEvents = [&](const FVisibleItem& VisibleResource, TArrayView<const FRDGPassHandle> Passes)
+	{
+		AddEvent(Builder, VisibleResource);
+
+		for (FRDGPassHandle PassHandle : Passes)
+		{
+			const FPassPacket& Pass = *Graph.GetPass(PassHandle);
+			const FVisiblePass& VisiblePass = VisibleGraph->GetVisiblePass(Pass);
+			AddEvent(Builder, VisiblePass);
+
+			const float Y = VisiblePass.Max.Y;
+
+			FSpline Spline;
+			Spline.Start.X = (VisiblePass.Max.X + VisiblePass.Min.X) * 0.5f;
+			Spline.Start.Y = Y;
+			Spline.StartDir = FVector2D(0, -1);
+			Spline.End = FVector2D(0, VisibleResource.Min.Y - Y);
+			Spline.EndDir = FVector2D(0, 1);
+			Spline.Thickness = 1.0f;
+			Spline.Tint = FLinearColor(0.8f, 0.8f, 0.8f, 0.7f);
+			Splines.Add(Spline);
+
+			for (const FVisibleScope& VisibleScope : VisibleGraph->Scopes)
+			{
+				const FScopePacket& Scope = VisibleScope.GetPacket();
+				if (Intersects(Scope, Pass))
+				{
+					VisibleItems.Add(&VisibleScope);
+				}
+			}
+		}
+	};
+
+	if (ShowTextures())
+	{
+		for (const FVisibleTexture& VisibleTexture : VisibleGraph->Textures)
+		{
+			const FTexturePacket& Texture = VisibleTexture.GetPacket();
+
+			if (FilterPacketExact(Texture))
+			{
+				AddResourcePassEvents(VisibleTexture, Texture.Passes);
+			}
+		}
+	}
+
+	if (ShowBuffers())
+	{
+		for (const FVisibleBuffer& VisibleBuffer : VisibleGraph->Buffers)
+		{
+			const FBufferPacket& Buffer = VisibleBuffer.GetPacket();
+
+			if (FilterPacketExact(Buffer))
+			{
+				AddResourcePassEvents(VisibleBuffer, Buffer.Passes);
+			}
+		}
+	}
+
+	for (const FVisibleItem* VisibleItem : VisibleItems)
+	{
+		AddEvent(Builder, *VisibleItem);
+	}
+}
+
+const TSharedPtr<const ITimingEvent> FRenderGraphTrack::GetEvent(float InPosX, float InPosY, const FTimingTrackViewport& Viewport) const
+{
+	TSharedPtr<FEvent> Event;
+
+	for (const FVisibleGraph& Graph : VisibleGraphs)
+	{
+		float AdjustedPosY = InPosY - GetPosY();
+
+		if (Graph.Intersects(InPosX, AdjustedPosY))
+		{
+			if (const FVisibleItem* Item = Graph.FindItem(InPosX, AdjustedPosY))
+			{
+				const FPacket& Packet = *Item->Packet;
+
+				if (Packet.Is<FScopePacket>())
+				{
+					Event = MakeShared<FVisibleScopeEvent>(SharedThis(this), *static_cast<const FVisibleScope*>(Item));
+				}
+				else if (Packet.Is<FPassPacket>())
+				{
+					Event = MakeShared<FVisiblePassEvent>(SharedThis(this), *static_cast<const FVisiblePass*>(Item));
+				}
+				else if (ShowTextures() && Packet.Is<FTexturePacket>())
+				{
+					Event = MakeShared<FVisibleTextureEvent>(SharedThis(this), *static_cast<const FVisibleTexture*>(Item));
+				}
+				else if (ShowBuffers() && Packet.Is<FBufferPacket>())
+				{
+					Event = MakeShared<FVisibleBufferEvent>(SharedThis(this), *static_cast<const FVisibleBuffer*>(Item));
+				}
+			}
+		}
+	}
+
+	return Event;
+}
+
+TSharedPtr<ITimingEventFilter> FRenderGraphTrack::GetFilterByEvent(const TSharedPtr<const ITimingEvent> InTimingEvent) const
+{
+	if (InTimingEvent.IsValid() && InTimingEvent->Is<FEvent>())
+	{
+		bool bFilterable = false;
+		if (InTimingEvent->As<FEvent>().GetItem().Intersects(MousePosition.X, MousePosition.Y, bFilterable) && bFilterable)
+		{
+			return MakeShared<FPacketFilter>(StaticCastSharedPtr<const FEvent>(InTimingEvent));
+		}
+	}
+	return nullptr;
+}
+
+const TCHAR* GetDimensionName(ETextureDimension Dimension)
+{
+	switch (Dimension)
+	{
+	case ETextureDimension::Texture2D:
+		return TEXT("Texture2D");
+	case ETextureDimension::Texture2DArray:
+		return TEXT("Texture2DArray");
+	case ETextureDimension::Texture3D:
+		return TEXT("Texture3D");
+	case ETextureDimension::TextureCube:
+		return TEXT("TextureCube");
+	case ETextureDimension::TextureCubeArray:
+		return TEXT("TextureCubeArray");
+	}
+	return TEXT("");
+}
+
+void FRenderGraphTrack::BuildContextMenu(FMenuBuilder& MenuBuilder)
+{
+	Super::BuildContextMenu(MenuBuilder);
+
+	MenuBuilder.BeginSection("Layout", LOCTEXT("TrackLayoutMenuHeader", "Track Layout"));
+	{
+		MenuBuilder.AddMenuEntry
+		(
+			LOCTEXT("ShowAlls", "Show All"),
+			LOCTEXT("ShowAlls_Tooltip", "Show All resources in the lifetime view."),
+			FSlateIcon(),
+			FUIAction(
+				FExecuteAction::CreateLambda([this]()
+				{
+					ResourceShow = EResourceShow::All;
+					SetDirtyFlag();
+				}),
+				FCanExecuteAction(),
+				FIsActionChecked::CreateLambda([this]() { return ResourceShow == EResourceShow::All; })
+			),
+			NAME_None,
+			EUserInterfaceActionType::RadioButton
+		);
+
+		MenuBuilder.AddMenuEntry
+		(
+			LOCTEXT("ShowTextures", "Show Textures"),
+			LOCTEXT("ShowTextures_Tooltip", "Show Texture resources in the lifetime view."),
+			FSlateIcon(),
+			FUIAction(
+				FExecuteAction::CreateLambda([this]()
+				{
+					ResourceShow = EResourceShow::Textures;
+					SetDirtyFlag();
+				}),
+				FCanExecuteAction(),
+				FIsActionChecked::CreateLambda([this]() { return ResourceShow == EResourceShow::Textures; })
+			),
+			NAME_None,
+			EUserInterfaceActionType::RadioButton
+		);
+
+		MenuBuilder.AddMenuEntry
+		(
+			LOCTEXT("ShowBuffers", "Show Buffers"),
+			LOCTEXT("ShowBuffers_Tooltip", "Show Buffer resources in the lifetime view."),
+			FSlateIcon(),
+			FUIAction(
+				FExecuteAction::CreateLambda([this]()
+				{
+					ResourceShow = EResourceShow::Buffers;
+					SetDirtyFlag();
+				}),
+				FCanExecuteAction(),
+				FIsActionChecked::CreateLambda([this]() { return ResourceShow == EResourceShow::Buffers; })
+			),
+			NAME_None,
+			EUserInterfaceActionType::RadioButton
+		);
+	}
+	MenuBuilder.EndSection();
+}
+
+void FRenderGraphTrack::InitTooltip(FTooltipDrawState& Tooltip, const ITimingEvent& InTooltipEvent) const
+{
+	Tooltip.ResetContent();
+
+	if (InTooltipEvent.CheckTrack(this) && InTooltipEvent.Is<FEvent>())
+	{
+		const auto AddCommonResourceText = [&](const FResourcePacket& Resource)
+		{
+			if (Resource.bExtracted)
+			{
+				Tooltip.AddTextLine(TEXT("Extracted"), FLinearColor::Red);
+			}
+
+			if (Resource.bExternal)
+			{
+				Tooltip.AddTextLine(TEXT("External"), FLinearColor::Red);
+			}
+		};
+
+		if (InTooltipEvent.Is<FVisibleScopeEvent>())
+		{
+			const auto& TooltipEvent = InTooltipEvent.As<FVisibleScopeEvent>();
+			const FScopePacket& Scope = TooltipEvent.GetPacket();
+			Tooltip.AddTitle(Scope.Name);
+		}
+		else if (InTooltipEvent.Is<FVisiblePassEvent>())
+		{
+			const auto& TooltipEvent = InTooltipEvent.As<FVisiblePassEvent>();
+			const FPassPacket& Pass = TooltipEvent.GetPacket();
+
+			Tooltip.AddTitle(GetSanitizeName(Pass.Name));
+			Tooltip.AddNameValueTextLine(TEXT("Handle:"), FString::Printf(TEXT("%d"), Pass.Handle.GetIndex()));
+
+			if (Pass.bCulled)
+			{
+				Tooltip.AddTextLine(TEXT("Culled"), FLinearColor::Red);
+			}
+			else
+			{
+				Tooltip.AddNameValueTextLine(TEXT("Used Textures:"), FString::Printf(TEXT("%d"), Pass.Textures.Num()));
+				Tooltip.AddNameValueTextLine(TEXT("Used Buffers:"), FString::Printf(TEXT("%d"), Pass.Buffers.Num()));
+			}
+
+			if (Pass.bSkipRenderPassBegin || Pass.bSkipRenderPassEnd)
+			{
+				Tooltip.AddTextLine(TEXT("Merged RenderPass"), FLinearColor::Red);
+			}
+		}
+		else if (InTooltipEvent.Is<FVisibleTextureEvent>())
+		{
+			const auto& TooltipEvent = InTooltipEvent.As<FVisibleTextureEvent>();
+			const FTexturePacket& Texture = TooltipEvent.GetPacket();
+
+			Tooltip.AddTitle(Texture.Name);
+			Tooltip.AddNameValueTextLine(TEXT("Dimension:"), GetDimensionName(Texture.Desc.Dimension));
+			Tooltip.AddNameValueTextLine(TEXT("Create Flags:"), GetTextureCreateFlagsName(Texture.Desc.Flags));
+			Tooltip.AddNameValueTextLine(TEXT("Format:"), UEnum::GetValueAsString(Texture.Desc.Format));
+			Tooltip.AddNameValueTextLine(TEXT("Extent:"), FString::Printf(TEXT("%d, %d"), Texture.Desc.Extent.X, Texture.Desc.Extent.Y));
+			Tooltip.AddNameValueTextLine(TEXT("Depth:"), FString::Printf(TEXT("%d"), Texture.Desc.Depth));
+			Tooltip.AddNameValueTextLine(TEXT("Mips:"), FString::Printf(TEXT("%d"), Texture.Desc.NumMips));
+			Tooltip.AddNameValueTextLine(TEXT("Array Size:"), FString::Printf(TEXT("%d"), Texture.Desc.ArraySize));
+			Tooltip.AddNameValueTextLine(TEXT("Samples:"), FString::Printf(TEXT("%d"), Texture.Desc.NumSamples));
+			Tooltip.AddNameValueTextLine(TEXT("Used Passes:"), FString::Printf(TEXT("%d"), Texture.Passes.Num()));
+			AddCommonResourceText(Texture);
+		}
+		else if (InTooltipEvent.Is<FVisibleBufferEvent>())
+		{
+			const auto& TooltipEvent = InTooltipEvent.As<FVisibleBufferEvent>();
+			const FBufferPacket& Buffer = TooltipEvent.GetPacket();
+
+			Tooltip.AddTitle(Buffer.Name);
+			Tooltip.AddNameValueTextLine(TEXT("Usage Flags:"), GetBufferUsageFlagsName(Buffer.Desc.Usage));
+			Tooltip.AddNameValueTextLine(TEXT("Bytes Per Element:"), FString::Printf(TEXT("%d"), Buffer.Desc.BytesPerElement));
+			Tooltip.AddNameValueTextLine(TEXT("Elements:"), FString::Printf(TEXT("%d"), Buffer.Desc.NumElements));
+			Tooltip.AddNameValueTextLine(TEXT("Used Passes:"), FString::Printf(TEXT("%d"), Buffer.Passes.Num()));
+			AddCommonResourceText(Buffer);
+		}
+
+		Tooltip.UpdateLayout();
+	}
+}
+
+} //namespace RenderGraphInsights
+} //namespace UE
+
+#undef LOCTEXT_NAMESPACE
\ No newline at end of file
diff --git a/Engine/Plugins/RenderGraphInsights/Source/RenderGraphInsights/Private/RenderGraphTrack.h b/Engine/Plugins/RenderGraphInsights/Source/RenderGraphInsights/Private/RenderGraphTrack.h
new file mode 100644
--- /dev/null
+++ b/Engine/Plugins/RenderGraphInsights/Source/RenderGraphInsights/Private/RenderGraphTrack.h
@@ -0,0 +1,270 @@
+// Copyright Epic Games, Inc. All Rights Reserved.
+
+#pragma once
+
+#include "Insights/ViewModels/TimingEventsTrack.h"
+#include "Insights/ViewModels/TimingEvent.h"
+#include "Insights/ViewModels/TimingTrackViewport.h"
+#include "Insights/ViewModels/TooltipDrawState.h"
+#include "RenderGraphProvider.h"
+
+class FMenuBuilder;
+class ITimingTrackDrawContext;
+class ITimingTrackUpdateContext;
+class FTimingTrackViewport;
+
+namespace UE
+{
+namespace RenderGraphInsights
+{
+class FGraphPacket;
+class FPacketFilter;
+
+class FVisibleItem
+{
+public:
+	FVisibleItem(const FTimingTrackViewport& InViewport, const FPacket& InPacket, double InStartTime, double InEndTime, uint32 InDepth, uint32 InColor, float InMinY, float InMaxY)
+		: Packet(&InPacket)
+		, Name(*InPacket.Name)
+		, StartTime(InStartTime)
+		, EndTime(InEndTime)
+		, Depth(InDepth)
+		, Color(InColor)
+	{
+		Min.X = InViewport.TimeToSlateUnitsRounded(StartTime);
+		Min.Y = InMinY;
+		Max.X = InViewport.TimeToSlateUnitsRounded(EndTime);
+		Max.Y = InMaxY;
+	}
+
+	virtual bool Intersects(float InPosX, float InPosY, bool& bOutFilterable) const
+	{
+		bOutFilterable = true;
+		return InPosX >= Min.X && InPosX < Max.X&& InPosY > Min.Y && InPosY < Max.Y;
+	}
+
+	bool Intersects(float InPosX, float InPosY) const
+	{
+		bool bFilterable = false;
+		return Intersects(InPosX, InPosY, bFilterable);
+	}
+
+	const FPacket* Packet{};
+	const TCHAR* Name{};
+	double StartTime{};
+	double EndTime{};
+	uint32 Depth{};
+	uint32 Color{};
+	uint32 Index{};
+	FVector2D Min = FVector2D::ZeroVector;
+	FVector2D Max = FVector2D::ZeroVector;
+};
+
+template <typename InPacketType>
+class TVisibleItemHelper : public FVisibleItem
+{
+public:
+	using PacketType = InPacketType;
+
+	TVisibleItemHelper(const FTimingTrackViewport& Viewport, const PacketType& InPacket, double InStartTime, double InEndTime, uint32 InDepth, uint32 InColor)
+		: TVisibleItemHelper(Viewport, InPacket, InStartTime, InEndTime, InDepth, InColor, Viewport.GetLayout().GetLaneY(InDepth))
+	{}
+
+	TVisibleItemHelper(const FTimingTrackViewport& Viewport, const PacketType& InPacket, double InStartTime, double InEndTime, uint32 InDepth, uint32 InColor, float InMinY)
+		: FVisibleItem(Viewport, InPacket, InStartTime, InEndTime, InDepth, InColor, InMinY, InMinY + Viewport.GetLayout().EventH)
+	{}
+
+	TVisibleItemHelper(const FTimingTrackViewport& Viewport, const PacketType& InPacket, double InStartTime, double InEndTime, uint32 InDepth, uint32 InColor, float InMinY, float InMaxY)
+		: FVisibleItem(Viewport, InPacket, InStartTime, InEndTime, InDepth, InColor, InMinY, InMaxY)
+	{}
+
+	const PacketType& GetPacket() const
+	{
+		return static_cast<const PacketType&>(*Packet);
+	}
+};
+
+class FVisiblePass final : public TVisibleItemHelper<FPassPacket>
+{
+public:
+	FVisiblePass(const FTimingTrackViewport& Viewport, const FPassPacket& InPacket, double InStartTime, double InEndTime, uint32 InDepth, uint32 InColor)
+		: TVisibleItemHelper(Viewport, InPacket, InStartTime, InEndTime, InDepth, InColor)
+	{}
+
+	bool Intersects(float InPosX, float InPosY, bool& bOutFilterable) const override
+	{
+		// For hit-testing, treat the pass as unbounded along Y, so that it's a column.
+		if (InPosX >= Min.X && InPosX < Max.X && InPosY >= Min.Y)
+		{
+			bOutFilterable = InPosY < Max.Y;
+			return true;
+		}
+		return false;
+	}
+
+	using FVisibleItem::Intersects;
+};
+
+class FVisibleScope final : public TVisibleItemHelper<FScopePacket>
+{
+public:
+
+	FVisibleScope(const FTimingTrackViewport& Viewport, const FScopePacket& InPacket, double InStartTime, double InEndTime, uint32 InDepth, uint32 InColor)
+		: TVisibleItemHelper(Viewport, InPacket, InStartTime, InEndTime, InDepth, InColor)
+	{}
+};
+
+class FVisibleTexture final : public TVisibleItemHelper<FTexturePacket>
+{
+public:
+
+	FVisibleTexture(const FTimingTrackViewport& Viewport, const FTexturePacket& InPacket, double InStartTime, double InEndTime, uint32 InDepth, uint32 InColor)
+		: TVisibleItemHelper(Viewport, InPacket, InStartTime, InEndTime, InDepth, InColor)
+	{}
+};
+
+class FVisibleBuffer final : public TVisibleItemHelper<FBufferPacket>
+{
+public:
+
+	FVisibleBuffer(const FTimingTrackViewport& Viewport, const FBufferPacket& InPacket, double InStartTime, double InEndTime, uint32 InDepth, uint32 InColor)
+		: TVisibleItemHelper(Viewport, InPacket, InStartTime, InEndTime, InDepth, InColor)
+	{}
+};
+
+class FVisibleGraph final : public TVisibleItemHelper<FGraphPacket>
+{
+public:
+	FVisibleGraph(const FTimingTrackViewport& Viewport, const FGraphPacket& InGraph, const uint32 InColor, float InMaxY)
+		: TVisibleItemHelper(Viewport, InGraph, InGraph.StartTime, InGraph.EndTime, 0, InColor, 0.0f, InMaxY)
+	{}
+
+	void AddScope(const FVisibleScope& VisibleScope);
+	void AddPass(const FVisiblePass& VisiblePass);
+	void AddTexture(const FVisibleTexture& VisibleTexture);
+	void AddBuffer(const FVisibleBuffer& VisibleBuffer);
+
+	void Reset();
+
+	const FVisibleItem* FindItem(float InPosX, float InPosY) const;
+
+	const FVisibleScope& GetVisibleScope(const FScopePacket& Scope) const
+	{
+		return Scopes[Scope.VisibleIndex];
+	}
+
+	const FVisiblePass& GetVisiblePass(const FPassPacket& Pass) const
+	{
+		return Passes[Pass.VisibleIndex];
+	}
+
+	const FVisibleTexture& GetVisibleTexture(const FTexturePacket& Texture) const
+	{
+		return Textures[Texture.VisibleIndex];
+	}
+
+	const FVisibleBuffer& GetVisibleBuffer(const FBufferPacket& Buffer) const
+	{
+		return Buffers[Buffer.VisibleIndex];
+	}
+
+	TArray<FVisibleScope> Scopes;
+	TArray<FVisiblePass> Passes;
+	TArray<FVisibleTexture> Textures;
+	TArray<FVisibleBuffer> Buffers;
+	TArray<uint32> AsyncComputePasses;
+};
+
+class FRenderGraphTrack : public FTimingEventsTrack
+{
+	INSIGHTS_DECLARE_RTTI(FRenderGraphTrack, FTimingEventsTrack)
+	using Super = FTimingEventsTrack;
+
+public:
+	FRenderGraphTrack(const FRenderGraphTimingViewSession& InSharedData);
+
+	//~ Begin FTimingEventsTrack interface
+	void Update(const ITimingTrackUpdateContext& Context) override;
+	void Draw(const ITimingTrackDrawContext& Context) const override;
+	void PostDraw(const ITimingTrackDrawContext& Context) const override;
+	void BuildDrawState(ITimingEventsTrackDrawStateBuilder& Builder, const ITimingTrackUpdateContext& Context) override;
+	void BuildFilteredDrawState(ITimingEventsTrackDrawStateBuilder& Builder, const ITimingTrackUpdateContext& Context) override;
+	const TSharedPtr<const ITimingEvent> GetEvent(float InPosX, float InPosY, const FTimingTrackViewport& Viewport) const override;
+	TSharedPtr<ITimingEventFilter> GetFilterByEvent(const TSharedPtr<const ITimingEvent> InTimingEvent) const;
+	void BuildContextMenu(FMenuBuilder& MenuBuilder) override;
+	void InitTooltip(FTooltipDrawState& Tooltip, const ITimingEvent& HoveredTimingEvent) const override;
+	//~ End FTimingEventsTrack interface
+
+private:
+	const FRenderGraphTimingViewSession& SharedData;
+
+	FVisibleGraph& AddVisibleGraph(const FTimingTrackViewport& Viewport, const FGraphPacket& Graph, uint32 Color)
+	{
+		const uint32 VisibleIndex = VisibleGraphs.Num();
+		FVisibleGraph& VisibleGraph = VisibleGraphs.Emplace_GetRef(Viewport, Graph, Color, GetHeight());
+		check(VisibleGraph.Packet->VisibleIndex == kInvalidVisibleIndex);
+		VisibleGraph.Packet->VisibleIndex = VisibleIndex;
+		VisibleGraph.Index = VisibleIndex;
+		return VisibleGraph;
+	}
+
+	const FVisibleGraph* GetVisibleGraph(const FGraphPacket& Packet) const
+	{
+		if (Packet.VisibleIndex != kInvalidVisibleIndex)
+		{
+			return &VisibleGraphs[Packet.VisibleIndex];
+		}
+		return nullptr;
+	}
+
+	bool FilterPacket(const FPacket& Packet) const;
+	bool FilterPacketExact(const FPacket& Packet) const;
+
+	FVector2D MousePosition = FVector2D::ZeroVector;
+
+	TArray<FVisibleGraph> VisibleGraphs;
+
+	struct FSpline
+	{
+		float Thickness{};
+		FVector2D Start = FVector2D::ZeroVector;
+		FVector2D StartDir = FVector2D::ZeroVector;
+		FVector2D End = FVector2D::ZeroVector;
+		FVector2D EndDir = FVector2D::ZeroVector;
+		FLinearColor Tint = FLinearColor::White;
+	};
+
+	TArray<FSpline> Splines;
+
+	const FPacketFilter* PacketFilter{};
+	const FPassPacket* SelectedPass{};
+
+	bool ShowTextures() const
+	{
+		return EnumHasAnyFlags(ResourceShow, EResourceShow::Textures);
+	}
+
+	bool ShowBuffers() const
+	{
+		return EnumHasAnyFlags(ResourceShow, EResourceShow::Buffers);
+	}
+
+	enum class EResourceShow
+	{
+		Textures = 1 << 0,
+		Buffers  = 1 << 1,
+		All = Textures | Buffers
+	};
+	FRIEND_ENUM_CLASS_FLAGS(FRenderGraphTrack::EResourceShow);
+
+	static void AddEvent(ITimingEventsTrackDrawStateBuilder& Builder, const FVisibleItem& Item);
+
+	EResourceShow ResourceShow = EResourceShow::All;
+
+	FTooltipDrawState SelectedTooltipState;
+};
+
+ENUM_CLASS_FLAGS(FRenderGraphTrack::EResourceShow);
+
+} //namespace RenderGraphInsights
+} //namespace UE
\ No newline at end of file
diff --git a/Engine/Plugins/RenderGraphInsights/Source/RenderGraphInsights/RenderGraphInsights.Build.cs b/Engine/Plugins/RenderGraphInsights/Source/RenderGraphInsights/RenderGraphInsights.Build.cs
new file mode 100644
--- /dev/null
+++ b/Engine/Plugins/RenderGraphInsights/Source/RenderGraphInsights/RenderGraphInsights.Build.cs
@@ -0,0 +1,26 @@
+// Copyright Epic Games, Inc. All Rights Reserved.
+
+namespace UnrealBuildTool.Rules
+{
+	public class RenderGraphInsights : ModuleRules
+	{
+		public RenderGraphInsights(ReadOnlyTargetRules Target) : base(Target)
+		{
+			PrivateDependencyModuleNames.AddRange(
+			new string[]
+			{
+				"Core",
+				"CoreUObject",
+				"RenderCore",
+				"RHI",
+				"SlateCore",
+				"Slate",
+				"TraceLog",
+				"TraceAnalysis",
+				"TraceServices",
+				"TraceInsights"
+			});
+		}
+	}
+}
+
diff --git a/Engine/Programs/UnrealInsights/Config/DefaultEngine.ini b/Engine/Programs/UnrealInsights/Config/DefaultEngine.ini
--- a/Engine/Programs/UnrealInsights/Config/DefaultEngine.ini
+++ b/Engine/Programs/UnrealInsights/Config/DefaultEngine.ini
@@ -2,6 +2,7 @@
 +ProgramEnabledPlugins=GameplayInsights
 +ProgramEnabledPlugins=TraceDataFilters
 +ProgramEnabledPlugins=SlateInsights
++ProgramEnabledPlugins=RenderGraphInsights
 +ProgramEnabledPlugins=NetworkPredictionInsights
 +ProgramEnabledPlugins=VisualStudioSourceCodeAccess
 +ProgramEnabledPlugins=XcodeSourceCodeAccess
diff --git a/Engine/Source/Developer/TraceInsights/Private/Insights/ViewModels/TimingViewDrawHelper.h b/Engine/Source/Developer/TraceInsights/Private/Insights/ViewModels/TimingViewDrawHelper.h
--- a/Engine/Source/Developer/TraceInsights/Private/Insights/ViewModels/TimingViewDrawHelper.h
+++ b/Engine/Source/Developer/TraceInsights/Private/Insights/ViewModels/TimingViewDrawHelper.h
@@ -169,6 +169,8 @@ public:
 	virtual FLinearColor GetTrackNameTextColor(const FBaseTimingTrack& Track) const override;
 	virtual int32 GetHeaderBackgroundLayerId() const override { return ReservedLayerId + ToInt32(EDrawLayer::HeaderBackground); }
 	virtual int32 GetHeaderTextLayerId() const override { return ReservedLayerId + ToInt32(EDrawLayer::HeaderText); }
+	virtual int32 GetFirstLayerId() const override { return ReservedLayerId; }
+	virtual int32 GetNumLayerIds() const override { return ToInt32(EDrawLayer::Count); }
 
 	const FDrawContext& GetDrawContext() const { return DrawContext; }
 	const FTimingTrackViewport& GetViewport() const { return Viewport; }
diff --git a/Engine/Source/Developer/TraceInsights/Public/Insights/Common/PaintUtils.h b/Engine/Source/Developer/TraceInsights/Public/Insights/Common/PaintUtils.h
--- a/Engine/Source/Developer/TraceInsights/Public/Insights/Common/PaintUtils.h
+++ b/Engine/Source/Developer/TraceInsights/Public/Insights/Common/PaintUtils.h
@@ -92,6 +92,11 @@ struct TRACEINSIGHTS_API FDrawContext
 		FSlateDrawElement::MakeText(ElementList, LayerId, MAKE_PAINT_GEOMETRY_PT(Geometry, TextX, Y), Text, Font, DrawEffects, Color);
 	}
 
+	inline void DrawSpline(uint32 InLayer, const float X, const float Y, const FVector2D& InStart, const FVector2D& InStartDir, const FVector2D& InEnd, const FVector2D& InEndDir, float InThickness = 0.0f, const FLinearColor& InTint=FLinearColor::White) const
+	{
+		FSlateDrawElement::MakeSpline(ElementList, InLayer, MAKE_PAINT_GEOMETRY_PT(Geometry, X, Y), InStart, InStartDir, InEnd, InEndDir, InThickness, DrawEffects, InTint);
+	}
+
 	//inline void IncrementLayer() const { ++LayerId; }
 
 	// Accessors
diff --git a/Engine/Source/Developer/TraceInsights/Public/Insights/ViewModels/ITimingViewDrawHelper.h b/Engine/Source/Developer/TraceInsights/Public/Insights/ViewModels/ITimingViewDrawHelper.h
--- a/Engine/Source/Developer/TraceInsights/Public/Insights/ViewModels/ITimingViewDrawHelper.h
+++ b/Engine/Source/Developer/TraceInsights/Public/Insights/ViewModels/ITimingViewDrawHelper.h
@@ -18,4 +18,6 @@ public:
 	virtual FLinearColor GetTrackNameTextColor(const FBaseTimingTrack& Track) const = 0;
 	virtual int32 GetHeaderBackgroundLayerId() const = 0;
 	virtual int32 GetHeaderTextLayerId() const = 0;
+	virtual int32 GetFirstLayerId() const = 0;
+	virtual int32 GetNumLayerIds() const = 0;
 };
diff --git a/Engine/Source/Runtime/RenderCore/Private/RenderGraphBuilder.cpp b/Engine/Source/Runtime/RenderCore/Private/RenderGraphBuilder.cpp
--- a/Engine/Source/Runtime/RenderCore/Private/RenderGraphBuilder.cpp
+++ b/Engine/Source/Runtime/RenderCore/Private/RenderGraphBuilder.cpp
@@ -2,6 +2,7 @@
 
 #include "RenderGraphBuilder.h"
 #include "RenderGraphPrivate.h"
+#include "RenderGraphTrace.h"
 #include "RenderTargetPool.h"
 #include "RenderGraphResourcePool.h"
 #include "VisualizeTexture.h"
@@ -1311,6 +1312,7 @@ void FRDGBuilder::Execute()
 
 	IF_RDG_CPU_SCOPES(CPUScopeStacks.BeginExecute());
 	IF_RDG_GPU_SCOPES(GPUScopeStacks.BeginExecute());
+	IF_RDG_ENABLE_TRACE(Trace.OutputGraphBegin());
 
 	if (!GRDGImmediateMode)
 	{
@@ -1358,6 +1360,7 @@ void FRDGBuilder::Execute()
 		*Query.Value = Query.Key->PooledBuffer;
 	}
 
+	IF_RDG_ENABLE_TRACE(Trace.OutputGraphEnd(*this));
 	IF_RDG_GPU_SCOPES(GPUScopeStacks.Graphics.EndExecute());
 	IF_RDG_CPU_SCOPES(CPUScopeStacks.EndExecute());
 
@@ -1511,6 +1514,10 @@ void FRDGBuilder::SetupPassInternal(FRDGPass* Pass, FRDGPassHandle PassHandle, E
 	IF_RDG_CPU_SCOPES(Pass->CPUScopes = CPUScopeStacks.GetCurrentScopes());
 	IF_RDG_GPU_SCOPES(Pass->GPUScopes = GPUScopeStacks.GetCurrentScopes(PassPipeline));
 
+#if RDG_GPU_SCOPES && RDG_ENABLE_TRACE
+	Pass->TraceEventScope = GPUScopeStacks.GetCurrentScopes(ERHIPipeline::Graphics).Event;
+#endif
+
 #if RDG_GPU_SCOPES && RDG_ENABLE_DEBUG
 	if (const FRDGEventScope* Scope = Pass->GPUScopes.Event)
 	{
@@ -1782,6 +1789,11 @@ void FRDGBuilder::CollectPassBarriers(FRDGPassHandle PassHandle, FRDGPassHandle&
 		FRDGTextureRef Texture = TexturePair.Key;
 		AddTransition(PassHandle, Texture, TexturePair.Value.MergeState, LastUntrackedPassHandle);
 		Texture->bCulled = false;
+
+	#if RDG_ENABLE_TRACE
+		Pass->Textures.Add(Texture->Handle);
+		Texture->Passes.Add(PassHandle);
+	#endif
 	}
 
 	for (const auto& BufferPair : Pass->BufferStates)
@@ -1789,6 +1801,11 @@ void FRDGBuilder::CollectPassBarriers(FRDGPassHandle PassHandle, FRDGPassHandle&
 		FRDGBufferRef Buffer = BufferPair.Key;
 		AddTransition(PassHandle, Buffer, *BufferPair.Value.MergeState, LastUntrackedPassHandle);
 		Buffer->bCulled = false;
+
+	#if RDG_ENABLE_TRACE
+		Pass->Buffers.Add(Buffer->Handle);
+		Buffer->Passes.Add(PassHandle);
+	#endif
 	}
 }
 
diff --git a/Engine/Source/Runtime/RenderCore/Private/RenderGraphEvent.cpp b/Engine/Source/Runtime/RenderCore/Private/RenderGraphEvent.cpp
--- a/Engine/Source/Runtime/RenderCore/Private/RenderGraphEvent.cpp
+++ b/Engine/Source/Runtime/RenderCore/Private/RenderGraphEvent.cpp
@@ -9,7 +9,7 @@ RENDERCORE_API bool GetEmitRDGEvents()
 {
 	check(IsInRenderingThread());
 #if RDG_EVENTS != RDG_EVENTS_NONE
-	return GetEmitDrawEvents() || GRDGDebug;
+	return GetEmitDrawEvents() || GRDGDebug || UE_TRACE_CHANNELEXPR_IS_ENABLED(RDGChannel);
 #else
 	return false;
 #endif
diff --git a/Engine/Source/Runtime/RenderCore/Private/RenderGraphTrace.cpp b/Engine/Source/Runtime/RenderCore/Private/RenderGraphTrace.cpp
new file mode 100644
--- /dev/null
+++ b/Engine/Source/Runtime/RenderCore/Private/RenderGraphTrace.cpp
@@ -0,0 +1,258 @@
+// Copyright Epic Games, Inc. All Rights Reserved.
+
+#include "RenderGraphTrace.h"
+#include "RenderGraphBuilder.h"
+#include "RenderGraphPrivate.h"
+#include "Trace/Trace.inl"
+
+#if RDG_ENABLE_TRACE
+
+UE_TRACE_CHANNEL_DEFINE(RDGChannel)
+
+UE_TRACE_EVENT_BEGIN(RDGTrace, GraphMessage)
+	UE_TRACE_EVENT_FIELD(UE::Trace::WideString, Name)
+	UE_TRACE_EVENT_FIELD(uint64, StartCycles)
+	UE_TRACE_EVENT_FIELD(uint64, EndCycles)
+	UE_TRACE_EVENT_FIELD(uint16, PassCount)
+UE_TRACE_EVENT_END()
+
+UE_TRACE_EVENT_BEGIN(RDGTrace, GraphEndMessage)
+UE_TRACE_EVENT_END()
+
+UE_TRACE_EVENT_BEGIN(RDGTrace, PassMessage)
+	UE_TRACE_EVENT_FIELD(UE::Trace::WideString, Name)
+	UE_TRACE_EVENT_FIELD(uint64, StartCycles)
+	UE_TRACE_EVENT_FIELD(uint64, EndCycles)
+	UE_TRACE_EVENT_FIELD(uint16, Handle)
+	UE_TRACE_EVENT_FIELD(uint16, GraphicsForkPass)
+	UE_TRACE_EVENT_FIELD(uint16, GraphicsJoinPass)
+	UE_TRACE_EVENT_FIELD(uint16[], Textures)
+	UE_TRACE_EVENT_FIELD(uint16[], Buffers)
+	UE_TRACE_EVENT_FIELD(uint16, Flags)
+	UE_TRACE_EVENT_FIELD(uint16, Pipeline)
+	UE_TRACE_EVENT_FIELD(bool, IsCulled)
+	UE_TRACE_EVENT_FIELD(bool, IsAsyncComputeBegin)
+	UE_TRACE_EVENT_FIELD(bool, IsAsyncComputeEnd)
+	UE_TRACE_EVENT_FIELD(bool, SkipRenderPassBegin)
+	UE_TRACE_EVENT_FIELD(bool, SkipRenderPassEnd)
+UE_TRACE_EVENT_END()
+
+UE_TRACE_EVENT_BEGIN(RDGTrace, BufferMessage)
+	UE_TRACE_EVENT_FIELD(UE::Trace::WideString, Name)
+	UE_TRACE_EVENT_FIELD(uint32, UsageFlags)
+	UE_TRACE_EVENT_FIELD(uint32, BytesPerElement)
+	UE_TRACE_EVENT_FIELD(uint32, NumElements)
+	UE_TRACE_EVENT_FIELD(uint16, Handle)
+	UE_TRACE_EVENT_FIELD(uint16, NextOwnerHandle)
+	UE_TRACE_EVENT_FIELD(uint16[], Passes)
+	UE_TRACE_EVENT_FIELD(bool, IsExternal)
+	UE_TRACE_EVENT_FIELD(bool, IsExtracted)
+	UE_TRACE_EVENT_FIELD(bool, IsCulled)
+UE_TRACE_EVENT_END()
+
+UE_TRACE_EVENT_BEGIN(RDGTrace, TextureMessage)
+	UE_TRACE_EVENT_FIELD(UE::Trace::WideString, Name)
+	UE_TRACE_EVENT_FIELD(uint64, StartCycles)
+	UE_TRACE_EVENT_FIELD(uint64, EndCycles)
+	UE_TRACE_EVENT_FIELD(uint16, Handle)
+	UE_TRACE_EVENT_FIELD(uint16, NextOwnerHandle)
+	UE_TRACE_EVENT_FIELD(uint16[], Passes)
+	UE_TRACE_EVENT_FIELD(uint64, SizeInBytes)
+	UE_TRACE_EVENT_FIELD(uint64, CreateFlags)
+	UE_TRACE_EVENT_FIELD(uint32, Dimension)
+	UE_TRACE_EVENT_FIELD(uint32, Format)
+	UE_TRACE_EVENT_FIELD(uint32, ExtentX)
+	UE_TRACE_EVENT_FIELD(uint32, ExtentY)
+	UE_TRACE_EVENT_FIELD(uint16, Depth)
+	UE_TRACE_EVENT_FIELD(uint16, ArraySize)
+	UE_TRACE_EVENT_FIELD(uint8, NumMips)
+	UE_TRACE_EVENT_FIELD(uint8, NumSamples)
+	UE_TRACE_EVENT_FIELD(bool, IsExternal)
+	UE_TRACE_EVENT_FIELD(bool, IsExtracted)
+	UE_TRACE_EVENT_FIELD(bool, IsCulled)
+UE_TRACE_EVENT_END()
+
+UE_TRACE_EVENT_BEGIN(RDGTrace, ScopeMessage)
+	UE_TRACE_EVENT_FIELD(UE::Trace::WideString, Name)
+	UE_TRACE_EVENT_FIELD(uint16, FirstPass)
+	UE_TRACE_EVENT_FIELD(uint16, LastPass)
+	UE_TRACE_EVENT_FIELD(uint16, Depth)
+UE_TRACE_EVENT_END()
+
+static_assert(sizeof(FRDGPassHandle) == sizeof(uint16), "Expected 16 bit pass handles.");
+static_assert(sizeof(FRDGTextureHandle) == sizeof(uint16), "Expected 16 bit texture handles.");
+static_assert(sizeof(FRDGBufferHandle) == sizeof(uint16), "Expected 16 bit buffer handles.");
+
+bool IsTraceEnabled()
+{
+	return UE_TRACE_CHANNELEXPR_IS_ENABLED(RDGChannel) && !GRDGImmediateMode;
+}
+
+void FRDGTrace::OutputGraphBegin()
+{
+	if (!IsTraceEnabled())
+	{
+		return;
+	}
+
+	GraphStartCycles = FPlatformTime::Cycles64();
+}
+
+void FRDGTrace::OutputGraphEnd(const FRDGBuilder& GraphBuilder)
+{
+	if (!IsTraceEnabled())
+	{
+		return;
+	}
+
+	const FRDGPassHandle ProloguePassHandle = GraphBuilder.GetProloguePassHandle();
+
+	const auto& Passes = GraphBuilder.Passes;
+	const auto& Textures = GraphBuilder.Textures;
+	const auto& Buffers = GraphBuilder.Buffers;
+
+	{
+		const TCHAR* Name = GraphBuilder.BuilderName.GetTCHAR();
+
+		UE_TRACE_LOG(RDGTrace, GraphMessage, RDGChannel)
+			<< GraphMessage.Name(Name, uint16(FCString::Strlen(Name)))
+			<< GraphMessage.StartCycles(GraphStartCycles)
+			<< GraphMessage.EndCycles(FPlatformTime::Cycles64())
+			<< GraphMessage.PassCount(uint16(Passes.Num()));
+	}
+
+	for (FRDGPassHandle Handle = Passes.Begin(); Handle != Passes.End(); ++Handle)
+	{
+		const FRDGPass* Pass = Passes[Handle];
+		const TCHAR* Name = Pass->GetEventName().GetTCHAR();
+
+		UE_TRACE_LOG(RDGTrace, PassMessage, RDGChannel)
+			<< PassMessage.Name(Name, uint16(FCString::Strlen(Name)))
+			<< PassMessage.Handle(Handle.GetIndex())
+			<< PassMessage.GraphicsForkPass(Pass->GetGraphicsForkPass().GetIndexUnchecked())
+			<< PassMessage.GraphicsJoinPass(Pass->GetGraphicsJoinPass().GetIndexUnchecked())
+			<< PassMessage.Textures((const uint16*)Pass->Textures.GetData(), (uint16)Pass->Textures.Num())
+			<< PassMessage.Buffers((const uint16*)Pass->Buffers.GetData(), (uint16)Pass->Buffers.Num())
+			<< PassMessage.Flags(uint16(Pass->GetFlags()))
+			<< PassMessage.Pipeline(uint16(Pass->GetPipeline()))
+			<< PassMessage.IsCulled(GraphBuilder.PassesToCull[Handle])
+			<< PassMessage.IsAsyncComputeBegin(Pass->bAsyncComputeBegin != 0)
+			<< PassMessage.IsAsyncComputeEnd(Pass->bAsyncComputeEnd != 0)
+			<< PassMessage.SkipRenderPassBegin(Pass->bSkipRenderPassBegin != 0)
+			<< PassMessage.SkipRenderPassEnd(Pass->bSkipRenderPassEnd != 0);
+	}
+
+#if RDG_EVENTS
+	{
+		struct FScopeInfo
+		{
+			const TCHAR* Name{};
+			FRDGPassHandle FirstPass;
+			FRDGPassHandle LastPass;
+			uint32 Depth{};
+		};
+		TArray<FScopeInfo> Scopes;
+		TMap<const FRDGEventScope*, int32> ScopeToIndex;
+		int32 Depth = 0;
+
+		TRDGScopeStackHelper<FRDGEventScope> ScopeStackHelper;
+
+		for (FRDGPassHandle Handle = Passes.Begin(); Handle != Passes.End(); ++Handle)
+		{
+			const FRDGPass* Pass = Passes[Handle];
+
+			const FRDGEventScope* ParentScope = Pass->TraceEventScope;
+
+			const auto PushScope = [&](const FRDGEventScope* ScopeToPush)
+			{
+				ScopeToIndex.Emplace(ScopeToPush, Scopes.Num());
+
+				FScopeInfo ScopeInfo;
+				ScopeInfo.Name = ScopeToPush->Name.GetTCHAR();
+				ScopeInfo.FirstPass = Handle;
+				ScopeInfo.Depth = Depth;
+				Scopes.Add(ScopeInfo);
+
+				Depth++;
+			};
+
+			const auto PopScope = [&](const FRDGEventScope* ScopeToPop)
+			{
+				FScopeInfo& ScopeInfo = Scopes[ScopeToIndex.FindChecked(ScopeToPop)];
+				ScopeInfo.LastPass = FRDGPassHandle(Handle.GetIndex() - 1);
+
+				Depth--;
+			};
+
+			ScopeStackHelper.BeginExecutePass(ParentScope, PushScope, PopScope);
+
+			if (Handle == Passes.Last())
+			{
+				ScopeStackHelper.EndExecute(PopScope);
+			}
+		}
+
+		check(Depth == 0);
+
+		for (const FScopeInfo& ScopeInfo : Scopes)
+		{
+			UE_TRACE_LOG(RDGTrace, ScopeMessage, RDGChannel)
+				<< ScopeMessage.Name(ScopeInfo.Name, uint16(FCString::Strlen(ScopeInfo.Name)))
+				<< ScopeMessage.FirstPass(ScopeInfo.FirstPass.GetIndexUnchecked())
+				<< ScopeMessage.LastPass(ScopeInfo.LastPass.GetIndexUnchecked())
+				<< ScopeMessage.Depth(ScopeInfo.Depth);
+		}
+	}
+#endif
+
+	for (FRDGTextureHandle Handle = Textures.Begin(); Handle != Textures.End(); ++Handle)
+	{
+		const FRDGTexture* Texture = Textures[Handle];
+
+		uint64 SizeInBytes = 0;
+		if (FRHITexture* TextureRHI = Texture->GetRHIUnchecked())
+		{
+			SizeInBytes = RHIComputeMemorySize(TextureRHI);
+		}
+
+		UE_TRACE_LOG(RDGTrace, TextureMessage, RDGChannel)
+			<< TextureMessage.Name(Texture->Name, uint16(FCString::Strlen(Texture->Name)))
+			<< TextureMessage.Handle(Handle.GetIndex())
+			<< TextureMessage.NextOwnerHandle(Texture->NextOwner.GetIndexUnchecked())
+			<< TextureMessage.Passes((const uint16*)Texture->Passes.GetData(), (uint16)Texture->Passes.Num())
+			<< TextureMessage.SizeInBytes(SizeInBytes)
+			<< TextureMessage.CreateFlags(uint32(Texture->Desc.Flags))
+			<< TextureMessage.Dimension(uint32(Texture->Desc.Dimension))
+			<< TextureMessage.Format(uint32(Texture->Desc.Format))
+			<< TextureMessage.ExtentX(Texture->Desc.Extent.X)
+			<< TextureMessage.ExtentY(Texture->Desc.Extent.Y)
+			<< TextureMessage.Depth(Texture->Desc.Depth)
+			<< TextureMessage.ArraySize(Texture->Desc.ArraySize)
+			<< TextureMessage.NumMips(Texture->Desc.NumMips)
+			<< TextureMessage.NumSamples(Texture->Desc.NumSamples)
+			<< TextureMessage.IsExternal(bool(Texture->bExternal))
+			<< TextureMessage.IsExtracted(bool(Texture->bExtracted))
+			<< TextureMessage.IsCulled(bool(Texture->bCulled));
+	}
+
+	for (FRDGBufferHandle Handle = Buffers.Begin(); Handle != Buffers.End(); ++Handle)
+	{
+		const FRDGBuffer* Buffer = Buffers[Handle];
+
+		UE_TRACE_LOG(RDGTrace, BufferMessage, RDGChannel)
+			<< BufferMessage.Name(Buffer->Name, uint16(FCString::Strlen(Buffer->Name)))
+			<< BufferMessage.Handle(Buffer->Handle.GetIndex())
+			<< BufferMessage.NextOwnerHandle(Buffer->NextOwner.GetIndexUnchecked())
+			<< BufferMessage.Passes((const uint16*)Buffer->Passes.GetData(), (uint16)Buffer->Passes.Num())
+			<< BufferMessage.UsageFlags(uint32(Buffer->Desc.Usage))
+			<< BufferMessage.BytesPerElement(Buffer->Desc.BytesPerElement)
+			<< BufferMessage.NumElements(Buffer->Desc.NumElements)
+			<< BufferMessage.IsExternal(bool(Buffer->bExternal))
+			<< BufferMessage.IsExtracted(bool(Buffer->bExtracted))
+			<< BufferMessage.IsCulled(bool(Buffer->bCulled));
+	}
+
+	UE_TRACE_LOG(RDGTrace, GraphEndMessage, RDGChannel);
+}
+
+#endif
\ No newline at end of file
diff --git a/Engine/Source/Runtime/RenderCore/Public/RenderGraphBuilder.h b/Engine/Source/Runtime/RenderCore/Public/RenderGraphBuilder.h
--- a/Engine/Source/Runtime/RenderCore/Public/RenderGraphBuilder.h
+++ b/Engine/Source/Runtime/RenderCore/Public/RenderGraphBuilder.h
@@ -6,6 +6,7 @@
 #include "RenderGraphEvent.h"
 #include "RenderGraphResources.h"
 #include "RenderGraphPass.h"
+#include "RenderGraphTrace.h"
 #include "RenderGraphValidation.h"
 #include "RenderGraphBlackboard.h"
 #include "ShaderParameterMacros.h"
@@ -295,6 +296,8 @@ private:
 	IF_RDG_CPU_SCOPES(FRDGCPUScopeStacks CPUScopeStacks);
 	IF_RDG_GPU_SCOPES(FRDGGPUScopeStacksByPipeline GPUScopeStacks);
 
+	IF_RDG_ENABLE_TRACE(FRDGTrace Trace);
+
 #if RDG_ENABLE_DEBUG
 	FRDGUserValidation UserValidation;
 	FRDGBarrierValidation BarrierValidation;
@@ -363,6 +366,7 @@ private:
 	void ClobberPassOutputs(const FRDGPass* Pass);
 #endif
 
+	friend FRDGTrace;
 	friend FRDGEventScopeGuard;
 	friend FRDGGPUStatScopeGuard;
 	friend FRDGAsyncComputeBudgetScopeGuard;
diff --git a/Engine/Source/Runtime/RenderCore/Public/RenderGraphDefinitions.h b/Engine/Source/Runtime/RenderCore/Public/RenderGraphDefinitions.h
--- a/Engine/Source/Runtime/RenderCore/Public/RenderGraphDefinitions.h
+++ b/Engine/Source/Runtime/RenderCore/Public/RenderGraphDefinitions.h
@@ -20,6 +20,15 @@
 /** Whether render graph debugging is enabled and we are compiling with the engine. */
 #define RDG_ENABLE_DEBUG_WITH_ENGINE (RDG_ENABLE_DEBUG && WITH_ENGINE)
 
+/** Whether render graph insight tracing is enabled. */
+#define RDG_ENABLE_TRACE UE_TRACE_ENABLED && !IS_PROGRAM && !UE_BUILD_SHIPPING
+
+#if RDG_ENABLE_TRACE
+	#define IF_RDG_ENABLE_TRACE(Op) Op
+#else
+	#define IF_RDG_ENABLE_TRACE(Op)
+#endif
+
 /** The type of GPU events the render graph system supports.
  *  RDG_EVENTS == 0 means there is no string processing at all.
  *  RDG_EVENTS == 1 means the format component of the event name is stored as a const TCHAR*.
@@ -234,11 +243,12 @@ public:
 
 	explicit inline TRDGHandle(int32 InIndex)
 	{
-		check(InIndex >= 0 && InIndex < kNullIndex);
+		check(InIndex >= 0 && InIndex <= kNullIndex);
 		Index = InIndex;
 	}
 
 	FORCEINLINE IndexType GetIndex() const { check(IsValid()); return Index; }
+	FORCEINLINE IndexType GetIndexUnchecked() const { return Index; }
 	FORCEINLINE bool IsNull()  const { return Index == kNullIndex; }
 	FORCEINLINE bool IsValid() const { return Index != kNullIndex; }
 	FORCEINLINE bool operator==(TRDGHandle Other) const { return Index == Other.Index; }
@@ -498,4 +508,6 @@ template <typename TUniformStruct> using TRDGUniformBufferRef = TRDGUniformBuffe
 template <typename InElementType, typename InAllocatorType = FDefaultAllocator>
 using TRDGTextureSubresourceArray = TArray<InElementType, TInlineAllocator<1, InAllocatorType>>;
 
-using FRDGPassHandlesByPipeline = TRHIPipelineArray<FRDGPassHandle>;
\ No newline at end of file
+using FRDGPassHandlesByPipeline = TRHIPipelineArray<FRDGPassHandle>;
+
+class FRDGTrace;
\ No newline at end of file
diff --git a/Engine/Source/Runtime/RenderCore/Public/RenderGraphEvent.h b/Engine/Source/Runtime/RenderCore/Public/RenderGraphEvent.h
--- a/Engine/Source/Runtime/RenderCore/Public/RenderGraphEvent.h
+++ b/Engine/Source/Runtime/RenderCore/Public/RenderGraphEvent.h
@@ -49,11 +49,29 @@
 /** Returns whether the current frame is emitting render graph events. */
 RENDERCORE_API bool GetEmitRDGEvents();
 
+template <typename TScopeType>
+class TRDGScopeStackHelper
+{
+	static constexpr uint32 kScopeStackDepthMax = 8;
+public:
+	TRDGScopeStackHelper()
+		: ScopeStack(MakeUniformStaticArray<const TScopeType*, kScopeStackDepthMax>(nullptr))
+	{}
+
+	template <typename PushFunctionType, typename PopFunctionType>
+	void BeginExecutePass(const TScopeType* ParentScope, PushFunctionType PushFunction, PopFunctionType PopFunction);
+
+	template <typename PopFunctionType>
+	void EndExecute(PopFunctionType PopFunction);
+
+private:
+	TStaticArray<const TScopeType*, kScopeStackDepthMax> ScopeStack;
+};
+
 /** A helper profiler class for tracking and evaluating hierarchical scopes in the context of render graph. */
 template <typename TScopeType>
 class TRDGScopeStack final
 {
-	static constexpr uint32 kScopeStackDepthMax = 8;
 public:
 	using FPushFunction = void(*)(FRHIComputeCommandList&, const TScopeType*);
 	using FPopFunction = void(*)(FRHIComputeCommandList&, const TScopeType*);
@@ -105,8 +123,7 @@ private:
 	/** Tracks scopes allocated through MemStack for destruction. */
 	TArray<TScopeType*, FRDGArrayAllocator> Scopes;
 
-	/** Stacks of scopes pushed to the RHI command list during execution. */
-	TStaticArray<const TScopeType*, kScopeStackDepthMax> ScopeStack;
+	TRDGScopeStackHelper<TScopeType> Helper;
 };
 
 class FRDGPass;
diff --git a/Engine/Source/Runtime/RenderCore/Public/RenderGraphEvent.inl b/Engine/Source/Runtime/RenderCore/Public/RenderGraphEvent.inl
--- a/Engine/Source/Runtime/RenderCore/Public/RenderGraphEvent.inl
+++ b/Engine/Source/Runtime/RenderCore/Public/RenderGraphEvent.inl
@@ -3,47 +3,8 @@
 #pragma once
 
 template <typename TScopeType>
-TRDGScopeStack<TScopeType>::TRDGScopeStack(
-	FRHIComputeCommandList& InRHICmdList,
-	FPushFunction InPushFunction,
-	FPopFunction InPopFunction)
-	: RHICmdList(InRHICmdList)
-	, MemStack(FMemStack::Get())
-	, PushFunction(InPushFunction)
-	, PopFunction(InPopFunction)
-	, ScopeStack(MakeUniformStaticArray<const TScopeType*, kScopeStackDepthMax>(nullptr))
-{}
-
-template <typename TScopeType>
-TRDGScopeStack<TScopeType>::~TRDGScopeStack()
+template <typename PushFunctionType, typename PopFunctionType>
+void TRDGScopeStackHelper<TScopeType>::BeginExecutePass(const TScopeType* ParentScope, PushFunctionType PushFunction, PopFunctionType PopFunction)
 {
-	ClearScopes();
-}
-
-template <typename TScopeType>
-template <typename... TScopeConstructArgs>
-void TRDGScopeStack<TScopeType>::BeginScope(TScopeConstructArgs... ScopeConstructArgs)
-{
-	auto Scope = new(MemStack) TScopeType(CurrentScope, Forward<TScopeConstructArgs>(ScopeConstructArgs)...);
-	Scopes.Add(Scope);
-	CurrentScope = Scope;
-}
-
-template <typename TScopeType>
-void TRDGScopeStack<TScopeType>::EndScope()
-{
-	checkf(CurrentScope != nullptr, TEXT("Current scope is null."));
-	CurrentScope = CurrentScope->ParentScope;
-}
-
-template <typename TScopeType>
-void TRDGScopeStack<TScopeType>::BeginExecute()
-{
-	checkf(CurrentScope == nullptr, TEXT("Render graph needs to have all scopes ended to execute."));
-}
-
-template <typename TScopeType>
-void TRDGScopeStack<TScopeType>::BeginExecutePass(const TScopeType* ParentScope)
-{
 	// Find out how many scopes needs to be popped.
 	TStaticArray<const TScopeType*, kScopeStackDepthMax> TraversedScopes;
@@ -81,21 +42,22 @@
 			break;
 		}
 
-		PopFunction(RHICmdList, ScopeStack[i]);
+		PopFunction(ScopeStack[i]);
 		ScopeStack[i] = nullptr;
 	}
 
 	// Push new scopes.
 	for (int32 i = TraversedScopeCount - 1; i >= 0 && CommonScopeId + 1 < static_cast<int32>(kScopeStackDepthMax); i--)
 	{
-		PushFunction(RHICmdList, TraversedScopes[i]);
+		PushFunction(TraversedScopes[i]);
 		CommonScopeId++;
 		ScopeStack[CommonScopeId] = TraversedScopes[i];
 	}
 }
 
 template <typename TScopeType>
-void TRDGScopeStack<TScopeType>::EndExecute()
+template <typename PopFunctionType>
+void TRDGScopeStackHelper<TScopeType>::EndExecute(PopFunctionType PopFunction)
 {
 	for (uint32 ScopeIndex = 0; ScopeIndex < kScopeStackDepthMax; ++ScopeIndex)
 	{
@@ -104,8 +66,71 @@
 			break;
 		}
 
-		PopFunction(RHICmdList, ScopeStack[ScopeIndex]);
+		PopFunction(ScopeStack[ScopeIndex]);
 	}
+}
+
+template <typename TScopeType>
+TRDGScopeStack<TScopeType>::TRDGScopeStack(
+	FRHIComputeCommandList& InRHICmdList,
+	FPushFunction InPushFunction,
+	FPopFunction InPopFunction)
+	: RHICmdList(InRHICmdList)
+	, MemStack(FMemStack::Get())
+	, PushFunction(InPushFunction)
+	, PopFunction(InPopFunction)
+{}
+
+template <typename TScopeType>
+TRDGScopeStack<TScopeType>::~TRDGScopeStack()
+{
 	ClearScopes();
 }
 
+template <typename TScopeType>
+template <typename... TScopeConstructArgs>
+void TRDGScopeStack<TScopeType>::BeginScope(TScopeConstructArgs... ScopeConstructArgs)
+{
+	auto Scope = new(MemStack) TScopeType(CurrentScope, Forward<TScopeConstructArgs>(ScopeConstructArgs)...);
+	Scopes.Add(Scope);
+	CurrentScope = Scope;
+}
+
+template <typename TScopeType>
+void TRDGScopeStack<TScopeType>::EndScope()
+{
+	checkf(CurrentScope != nullptr, TEXT("Current scope is null."));
+	CurrentScope = CurrentScope->ParentScope;
+}
+
+template <typename TScopeType>
+void TRDGScopeStack<TScopeType>::BeginExecute()
+{
+	checkf(CurrentScope == nullptr, TEXT("Render graph needs to have all scopes ended to execute."));
+}
+
+template <typename ScopeType>
+void TRDGScopeStack<ScopeType>::BeginExecutePass(const ScopeType* ParentScope)
+{
+	Helper.BeginExecutePass(
+		ParentScope,
+		[this](const ScopeType* Scope)
+		{
+			PushFunction(RHICmdList, Scope);
+		},
+		[this](const ScopeType* Scope)
+		{
+			PopFunction(RHICmdList, Scope);
+		});
+}
+
+template <typename TScopeType>
+void TRDGScopeStack<TScopeType>::EndExecute()
+{
+	Helper.EndExecute([this](const TScopeType* Scope)
+	{
+		PopFunction(RHICmdList, Scope);
+	});
+	ClearScopes();
+}
+
diff --git a/Engine/Source/Runtime/RenderCore/Public/RenderGraphPass.h b/Engine/Source/Runtime/RenderCore/Public/RenderGraphPass.h
--- a/Engine/Source/Runtime/RenderCore/Public/RenderGraphPass.h
+++ b/Engine/Source/Runtime/RenderCore/Public/RenderGraphPass.h
@@ -379,8 +379,18 @@ private:
 	IF_RDG_CPU_SCOPES(FRDGCPUScopes CPUScopes);
 	IF_RDG_GPU_SCOPES(FRDGGPUScopes GPUScopes);
 
+#if RDG_GPU_SCOPES && RDG_ENABLE_TRACE
+	const FRDGEventScope* TraceEventScope = nullptr;
+#endif
+
+#if RDG_ENABLE_TRACE
+	TArray<FRDGTextureHandle, FRDGArrayAllocator> Textures;
+	TArray<FRDGBufferHandle, FRDGArrayAllocator> Buffers;
+#endif
+
 	friend FRDGBuilder;
 	friend FRDGPassRegistry;
+	friend FRDGTrace;
 };
 
 /** Render graph pass with lambda execute function. */
diff --git a/Engine/Source/Runtime/RenderCore/Public/RenderGraphResources.h b/Engine/Source/Runtime/RenderCore/Public/RenderGraphResources.h
--- a/Engine/Source/Runtime/RenderCore/Public/RenderGraphResources.h
+++ b/Engine/Source/Runtime/RenderCore/Public/RenderGraphResources.h
@@ -220,6 +220,12 @@ public:
 		return bExternal;
 	}
 
+	/** Whether this resource is has been queued for extraction at the end of graph execution. */
+	bool IsExtracted() const
+	{
+		return bExtracted;
+	}
+
 	/** Whether a prior pass added to the graph produced contents for this resource. External resources are not considered produced
 	 *  until used for a write operation. This is a union of all subresources, so any subresource write will set this to true.
 	 */
@@ -267,6 +273,10 @@ private:
 	FRDGPassHandle FirstPass;
 	FRDGPassHandle LastPass;
 
+#if RDG_ENABLE_TRACE
+	TArray<FRDGPassHandle, FRDGArrayAllocator> Passes;
+#endif
+
 #if RDG_ENABLE_DEBUG
 	struct FRDGParentResourceDebugData* ParentDebugData = nullptr;
 	FRDGParentResourceDebugData& GetParentDebugData() const;
@@ -275,6 +285,7 @@ private:
 	friend FRDGBuilder;
 	friend FRDGUserValidation;
 	friend FRDGBarrierBatchBegin;
+	friend FRDGTrace;
 };
 
 /** A render graph resource (e.g. a view) which references a single parent resource (e.g. a texture / buffer). Provides an abstract way to access the parent resource. */
@@ -609,6 +620,11 @@ public:
 
 	//////////////////////////////////////////////////////////////////////////
 
+	FRDGTextureHandle GetHandle() const
+	{
+		return Handle;
+	}
+
 	FRDGTextureSubresourceLayout GetSubresourceLayout() const
 	{
 		return Layout;
@@ -699,6 +715,7 @@ private:
 	friend FRDGTextureRegistry;
 	friend FRDGAllocator;
 	friend FPooledRenderTarget;
+	friend FRDGTrace;
 };
 
 /** Render graph tracked SRV. */
@@ -1189,6 +1206,11 @@ public:
 
 	//////////////////////////////////////////////////////////////////////////
 
+	FRDGBufferHandle GetHandle() const
+	{
+		return Handle;
+	}
+
 private:
 	FRDGBuffer(const TCHAR* InName, const FRDGBufferDesc& InDesc, ERDGBufferFlags InFlags)
 		: FRDGParentResource(InName, ERDGParentResourceType::Buffer)
@@ -1240,6 +1262,7 @@ private:
 	friend FRDGUserValidation;
 	friend FRDGBufferRegistry;
 	friend FRDGAllocator;
+	friend FRDGTrace;
 };
 
 /** Render graph tracked buffer SRV. */
diff --git a/Engine/Source/Runtime/RenderCore/Public/RenderGraphTrace.h b/Engine/Source/Runtime/RenderCore/Public/RenderGraphTrace.h
new file mode 100644
--- /dev/null
+++ b/Engine/Source/Runtime/RenderCore/Public/RenderGraphTrace.h
@@ -0,0 +1,21 @@
+// Copyright Epic Games, Inc. All Rights Reserved.
+
+#pragma once
+
+#include "RenderGraphDefinitions.h"
+
+#if RDG_ENABLE_TRACE
+
+UE_TRACE_CHANNEL_EXTERN(RDGChannel, RENDERCORE_API);
+
+class RENDERCORE_API FRDGTrace
+{
+public:
+	void OutputGraphBegin();
+	void OutputGraphEnd(const FRDGBuilder& GraphBuilder);
+
+private:
+	uint64 GraphStartCycles{};
+};
+
+#endif
\ No newline at end of file
