Non-unity compilation fixes.

diff --git a/Engine/Plugins/RenderGraphInsights/Source/RenderGraphInsights/Private/RenderGraphProvider.cpp b/Engine/Plugins/RenderGraphInsights/Source/RenderGraphInsights/Private/RenderGraphProvider.cpp
--- a/Engine/Plugins/RenderGraphInsights/Source/RenderGraphInsights/Private/RenderGraphProvider.cpp
+++ b/Engine/Plugins/RenderGraphInsights/Source/RenderGraphInsights/Private/RenderGraphProvider.cpp
@@ -18,20 +18,6 @@ INSIGHTS_IMPLEMENT_RTTI(FPassIntervalPacket)
 
 FName FRenderGraphProvider::ProviderName("RenderGraphProvider");
 
-void SanitizeName(FString& Name)
-{
-	if (Name.IsEmpty())
-	{
-		Name = TEXT("<unnamed>");
-	}
-}
-
-FString GetSanitizeName(FString Name)
-{
-	SanitizeName(Name);
-	return MoveTemp(Name);
-}
-
 FString GetSizeName(uint32 Bytes)
 {
 	const uint32 KB = 1024;
diff --git a/Engine/Plugins/RenderGraphInsights/Source/RenderGraphInsights/Private/RenderGraphProvider.h b/Engine/Plugins/RenderGraphInsights/Source/RenderGraphInsights/Private/RenderGraphProvider.h
--- a/Engine/Plugins/RenderGraphInsights/Source/RenderGraphInsights/Private/RenderGraphProvider.h
+++ b/Engine/Plugins/RenderGraphInsights/Source/RenderGraphInsights/Private/RenderGraphProvider.h
@@ -17,6 +17,20 @@ namespace UE
 namespace RenderGraphInsights
 {
 
+inline void SanitizeName(FString& Name)
+{
+	if (Name.IsEmpty())
+	{
+		Name = TEXT("<unnamed>");
+	}
+}
+
+inline FString GetSanitizedName(FString Name)
+{
+	SanitizeName(Name);
+	return MoveTemp(Name);
+}
+
 class FGraphPacket;
 
 const uint32 kInvalidVisibleIndex = ~0u;
diff --git a/Engine/Plugins/RenderGraphInsights/Source/RenderGraphInsights/Private/RenderGraphTrack.cpp b/Engine/Plugins/RenderGraphInsights/Source/RenderGraphInsights/Private/RenderGraphTrack.cpp
--- a/Engine/Plugins/RenderGraphInsights/Source/RenderGraphInsights/Private/RenderGraphTrack.cpp
+++ b/Engine/Plugins/RenderGraphInsights/Source/RenderGraphInsights/Private/RenderGraphTrack.cpp
@@ -1068,7 +1068,7 @@ void FRenderGraphTrack::InitTooltip(FTooltipDrawState& Tooltip, const ITimingEve
 			const auto& TooltipEvent = InTooltipEvent.As<FVisiblePassEvent>();
 			const FPassPacket& Pass = TooltipEvent.GetPacket();
 
-			Tooltip.AddTitle(GetSanitizeName(Pass.Name));
+			Tooltip.AddTitle(GetSanitizedName(Pass.Name));
 			Tooltip.AddNameValueTextLine(TEXT("Handle:"), FString::Printf(TEXT("%d"), Pass.Handle.GetIndex()));
 
 			if (Pass.bCulled)
diff --git a/Engine/Plugins/RenderGraphInsights/Source/RenderGraphInsights/Private/RenderGraphTrack.h b/Engine/Plugins/RenderGraphInsights/Source/RenderGraphInsights/Private/RenderGraphTrack.h
--- a/Engine/Plugins/RenderGraphInsights/Source/RenderGraphInsights/Private/RenderGraphTrack.h
+++ b/Engine/Plugins/RenderGraphInsights/Source/RenderGraphInsights/Private/RenderGraphTrack.h
@@ -17,6 +17,8 @@ namespace UE
 {
 namespace RenderGraphInsights
 {
+class FRenderGraphTimingViewSession;
+
 class FGraphPacket;
 class FPacketFilter;
 
