Improved sorting and coloration of resources in RDG insights.

#rnx

diff --git a/Engine/Plugins/RenderGraphInsights/Source/RenderGraphInsights/Private/RenderGraphProvider.cpp b/Engine/Plugins/RenderGraphInsights/Source/RenderGraphInsights/Private/RenderGraphProvider.cpp
--- a/Engine/Plugins/RenderGraphInsights/Source/RenderGraphInsights/Private/RenderGraphProvider.cpp
+++ b/Engine/Plugins/RenderGraphInsights/Source/RenderGraphInsights/Private/RenderGraphProvider.cpp
@@ -59,6 +59,8 @@ FScopePacket::FScopePacket(const UE::Trace::IAnalyzer::FOnEventContext& Context)
 
 FResourcePacket::FResourcePacket(const UE::Trace::IAnalyzer::FOnEventContext& Context)
 	: FPassIntervalPacket(Context)
+	, Order(Context.EventData.GetValue<uint16>("Order"))
+	, SizeInBytes(Context.EventData.GetValue<uint64>("SizeInBytes"))
 	, Passes(Context.EventData.GetArrayView<FRDGPassHandle>("Passes"))
 	, bExternal(Context.EventData.GetValue<bool>("IsExternal"))
 	, bExtracted(Context.EventData.GetValue<bool>("IsExtracted"))
@@ -75,7 +77,6 @@ FTexturePacket::FTexturePacket(const UE::Trace::IAnalyzer::FOnEventContext& Cont
 	: FResourcePacket(Context)
 	, Handle(Context.EventData.GetValue<FRDGTextureHandle>("Handle"))
 	, NextOwnerHandle(Context.EventData.GetValue<FRDGTextureHandle>("NextOwnerHandle"))
-	, SizeInBytes(Context.EventData.GetValue<uint64>("SizeInBytes"))
 {
 	Desc.Flags = ETextureCreateFlags(Context.EventData.GetValue<uint64>("CreateFlags"));
 	Desc.Dimension = ETextureDimension(Context.EventData.GetValue<uint16>("Dimension"));
@@ -99,7 +100,7 @@ FBufferPacket::FBufferPacket(const UE::Trace::IAnalyzer::FOnEventContext& Contex
 	Desc.BytesPerElement = Context.EventData.GetValue<uint32>("BytesPerElement");
 	Desc.NumElements = Context.EventData.GetValue<uint32>("NumElements");
 
-	const uint64 SizeInBytes = Desc.BytesPerElement * Desc.NumElements;
+	SizeInBytes = Desc.BytesPerElement * Desc.NumElements;
 	Name += GetSizeName(SizeInBytes);
 }
 
@@ -212,12 +213,7 @@ void FRenderGraphProvider::AddTexture(const FTexturePacket InTexture)
 	if (const FTexturePacket* const* PreviousOwnerPtr = CurrentGraph->TextureHandleToPreviousOwner.Find(Texture.Handle))
 	{
 		const FTexturePacket* PreviousOwner = *PreviousOwnerPtr;
-		Texture.PreviousOwner = PreviousOwner;
-		Texture.Depth = PreviousOwner->Depth;
-	}
-	else if (!Texture.bCulled)
-	{
-		Texture.Depth = CurrentGraph->TextureDepthOffset++;
+		Texture.PrevousOwnerHandle = PreviousOwner->Handle;
 	}
 
 	if (Texture.NextOwnerHandle.IsValid())
@@ -235,12 +231,7 @@ void FRenderGraphProvider::AddBuffer(const FBufferPacket InBuffer)
 	if (const FBufferPacket* const* PreviousOwnerPtr = CurrentGraph->BufferHandleToPreviousOwner.Find(Buffer.Handle))
 	{
 		const FBufferPacket* PreviousOwner = *PreviousOwnerPtr;
-		Buffer.PreviousOwner = PreviousOwner;
-		Buffer.Depth = PreviousOwner->Depth;
-	}
-	else if (!Buffer.bCulled)
-	{
-		Buffer.Depth = CurrentGraph->BufferDepthOffset++;
+		Buffer.PrevousOwnerHandle = PreviousOwner->Handle;
 	}
 
 	if (Buffer.NextOwnerHandle.IsValid())
diff --git a/Engine/Plugins/RenderGraphInsights/Source/RenderGraphInsights/Private/RenderGraphProvider.h b/Engine/Plugins/RenderGraphInsights/Source/RenderGraphInsights/Private/RenderGraphProvider.h
--- a/Engine/Plugins/RenderGraphInsights/Source/RenderGraphInsights/Private/RenderGraphProvider.h
+++ b/Engine/Plugins/RenderGraphInsights/Source/RenderGraphInsights/Private/RenderGraphProvider.h
@@ -74,7 +74,8 @@ class FResourcePacket : public FPassIntervalPacket
 	INSIGHTS_DECLARE_RTTI(FResourcePacket, FPassIntervalPacket)
 public:
 	uint32 Index{};
-	uint32 Depth{};
+	uint32 Order{};
+	uint64 SizeInBytes{};
 
 	TArray<FRDGPassHandle> Passes;
 
@@ -89,12 +90,10 @@ class FTexturePacket : public FResourcePacket
 {
 	INSIGHTS_DECLARE_RTTI(FTexturePacket, FResourcePacket)
 public:
-	const FTexturePacket* PreviousOwner{};
-
 	FRDGTextureHandle Handle;
 	FRDGTextureHandle NextOwnerHandle;
+	FRDGTextureHandle PrevousOwnerHandle;
 	FRDGTextureDesc Desc;
-	uint64 SizeInBytes{};
 
 	FTexturePacket(const UE::Trace::IAnalyzer::FOnEventContext& Context);
 };
@@ -103,10 +102,9 @@ class FBufferPacket : public FResourcePacket
 {
 	INSIGHTS_DECLARE_RTTI(FBufferPacket, FResourcePacket)
 public:
-	const FBufferPacket* PreviousOwner{};
-
 	FRDGBufferHandle Handle;
 	FRDGBufferHandle NextOwnerHandle;
+	FRDGBufferHandle PrevousOwnerHandle;
 	FRDGBufferDesc Desc;
 
 	FBufferPacket(const UE::Trace::IAnalyzer::FOnEventContext& Context);
@@ -148,8 +146,6 @@ public:
 
 	uint32 ScopeDepth{};
 	uint32 PassCount{};
-	uint32 TextureDepthOffset{};
-	uint32 BufferDepthOffset{};
 
 	FGraphPacket(TraceServices::ILinearAllocator& Allocator, const UE::Trace::IAnalyzer::FOnEventContext& Context);
 
diff --git a/Engine/Plugins/RenderGraphInsights/Source/RenderGraphInsights/Private/RenderGraphTrack.cpp b/Engine/Plugins/RenderGraphInsights/Source/RenderGraphInsights/Private/RenderGraphTrack.cpp
--- a/Engine/Plugins/RenderGraphInsights/Source/RenderGraphInsights/Private/RenderGraphTrack.cpp
+++ b/Engine/Plugins/RenderGraphInsights/Source/RenderGraphInsights/Private/RenderGraphTrack.cpp
@@ -19,6 +19,9 @@
 #include "Insights/ViewModels/TimingEventSearch.h"
 #include "Insights/ViewModels/TooltipDrawState.h"
 
+#include "Widgets/Input/SSearchBox.h"
+#include "Widgets/Input/SSpinBox.h"
+
 #include "TraceServices/Model/Frames.h"
 
 #define LOCTEXT_NAMESPACE "RenderGraphTrack"
@@ -181,6 +184,11 @@ static bool Intersects(const FPassIntervalPacket& A, const FPassIntervalPacket&
 	return !(B.LastPass < A.FirstPass || A.LastPass < B.FirstPass);
 }
 
+static void AddEvent(ITimingEventsTrackDrawStateBuilder& Builder, const FVisibleItem& Item)
+{
+	Builder.AddEvent(Item.StartTime, Item.EndTime, Item.Depth, Item.Color, [&Item](float) { return Item.Name; });
+};
+
 bool FPacketFilter::FilterPacketExact(const FPacket& PacketToFilter) const
 {
 	return Packet == &PacketToFilter;
@@ -347,6 +355,34 @@ static uint32 GetPassColor(const FPassPacket& Packet)
 	return Color;
 }
 
+static uint32 GetColorBySize(uint64 Size, uint64 MaxSize)
+{
+	const FLinearColor Low(0.01, 0.01, 0.01, 0.25f);
+	const FLinearColor High(1.0, 0.1, 0.1, 1.0f);
+	const float Percentage = FMath::Sqrt(float(Size) / float(MaxSize));
+	return FLinearColor::LerpUsingHSV(Low, High, Percentage).ToFColor(false).ToPackedARGB();
+}
+
+uint32 FRenderGraphTrack::GetTextureColor(const FTexturePacket& Texture, uint64 MaxSizeInBytes) const
+{
+	if (ResourceColor == EResourceColor::Type)
+	{
+		return kTextureColor;
+	}
+
+	return GetColorBySize(Texture.SizeInBytes, MaxSizeInBytes);
+}
+
+uint32 FRenderGraphTrack::GetBufferColor(const FBufferPacket& Texture, uint64 MaxSizeInBytes) const
+{
+	if (ResourceColor == EResourceColor::Type)
+	{
+		return kBufferColor;
+	}
+
+	return GetColorBySize(Texture.SizeInBytes, MaxSizeInBytes);
+}
+
 void FRenderGraphTrack::Update(const ITimingTrackUpdateContext& Context)
 {
 	FTimingEventsTrack::Update(Context);
@@ -601,11 +637,6 @@ void FRenderGraphTrack::PostDraw(const ITimingTrackDrawContext& Context) const
 	SelectedTooltipState.Draw(Context.GetDrawContext());
 }
 
-void FRenderGraphTrack::AddEvent(ITimingEventsTrackDrawStateBuilder& Builder, const FVisibleItem& Item)
-{
-	Builder.AddEvent(Item.StartTime, Item.EndTime, Item.Depth, Item.Color, [&Item](float) { return Item.Name; });
-};
-
 void FRenderGraphTrack::BuildDrawState(ITimingEventsTrackDrawStateBuilder& Builder, const ITimingTrackUpdateContext& Context)
 {
 	// Reset draw state; it's built each frame.
@@ -629,13 +660,28 @@ void FRenderGraphTrack::BuildDrawState(ITimingEventsTrackDrawStateBuilder& Build
 
 	if (const FRenderGraphProvider* RenderGraphProvider = SharedData.GetAnalysisSession().ReadProvider<FRenderGraphProvider>(FRenderGraphProvider::ProviderName))
 	{
+		struct FResourceEntry
+		{
+			double StartTime{};
+			double EndTime{};
+			uint64 SizeInBytes{};
+			uint32 Index{};
+			uint32 Order{};
+			ERDGParentResourceType Type = ERDGParentResourceType::Texture;
+			bool bHasPreviousOwner{};
+		};
+
+		TArray<FResourceEntry> Resources;
+		TArray<uint16> TextureIndexToDepth;
+		TArray<uint16> BufferIndexToDepth;
+
 		const FTimingTrackViewport& Viewport = Context.GetViewport();
 
 		TraceServices::FAnalysisSessionReadScope SessionReadScope(SharedData.GetAnalysisSession());
 
 		const FRenderGraphProvider::TGraphTimeline& GraphTimeline = RenderGraphProvider->GetGraphTimeline();
 		GraphTimeline.EnumerateEvents(Viewport.GetStartTime(), Viewport.GetEndTime(),
-			[this, &Builder, &Viewport](double GraphStartTime, double GraphEndTime, uint32, const TSharedPtr<FGraphPacket>& Graph)
+			[this, &Builder, &Viewport, &Resources, &TextureIndexToDepth, &BufferIndexToDepth](double GraphStartTime, double GraphEndTime, uint32, const TSharedPtr<FGraphPacket>& Graph)
 		{
 			// We always render the graph event, so add it separately even if the visible graph is culled.
 			Builder.AddEvent(GraphStartTime, GraphEndTime, 0, *Graph->Name, 0, kBuilderColor);
@@ -681,10 +727,10 @@ void FRenderGraphTrack::BuildDrawState(ITimingEventsTrackDrawStateBuilder& Build
 			{
 				FPassPacket& Pass = Graph->Passes[PassIndex];
 				const double StartTime = Pass.StartTime + SinglePixelTimeMargin;
-				const double EndTime   = Pass.EndTime;
+				const double EndTime = Pass.EndTime;
 				const bool bAsyncCompute = EnumHasAnyFlags(Pass.Flags, ERDGPassFlags::AsyncCompute);
-				const uint32 Depth     = DepthOffset + (bAsyncCompute ? 2 : 0);
-				const uint32 Color     = GetPassColor(Pass);
+				const uint32 Depth = DepthOffset + (bAsyncCompute ? 2 : 0);
+				const uint32 Color = GetPassColor(Pass);
 
 				const FVisiblePass VisiblePass(Viewport, Pass, StartTime, EndTime, Depth, Color);
 				AddEvent(Builder, VisiblePass);
@@ -696,60 +742,176 @@ void FRenderGraphTrack::BuildDrawState(ITimingEventsTrackDrawStateBuilder& Build
 			// Empty space between passes / resources.
 			DepthOffset += bAnyAsyncCompute ? 3 : 1;
 
-			if (ShowTextures())
+			Resources.Reset();
+			Resources.Reserve(Graph->Textures.Num() + Graph->Buffers.Num());
+			TextureIndexToDepth.SetNum(Graph->Textures.Num());
+			BufferIndexToDepth.SetNum(Graph->Buffers.Num());
+
+			TBitArray<TInlineAllocator<8>> CulledTextures(true, Graph->Textures.Num());
+
+			const auto CheckFilterName = [this](const FString& InName)
+			{
+				return FilterText.IsEmpty() || InName.Contains(FilterText);
+			};
+
+			const auto CheckFilterSize = [this](uint64 InSizeInBytes)
 			{
-				uint32 TextureDepthOffsetMax = 0;
+				return FilterSize <= 0.0f || InSizeInBytes >= uint64(FilterSize * 1024.0f * 1024.0f);
+			};
 
+			if (ShowTextures())
+			{
 				for (uint32 TextureIndex = 0, TextureCount = Graph->Textures.Num(); TextureIndex < TextureCount; ++TextureIndex)
 				{
 					FTexturePacket& Texture = Graph->Textures[TextureIndex];
-					if (Texture.bCulled)
+
+					const bool bCulled = Texture.bCulled || !CheckFilterName(Texture.Name) || !CheckFilterSize(Texture.SizeInBytes);
+					CulledTextures[TextureIndex] = bCulled;
+
+					if (bCulled)
+					{
+						continue;
+					}
+
+					FResourceEntry& Entry = Resources.AddDefaulted_GetRef();
+					Entry.StartTime = Texture.StartTime;
+					Entry.EndTime = Texture.EndTime;
+					Entry.SizeInBytes = Texture.SizeInBytes;
+					Entry.Index = TextureIndex;
+					Entry.Order = Texture.Order;
+					Entry.Type = ERDGParentResourceType::Texture;
+					Entry.bHasPreviousOwner = Texture.PrevousOwnerHandle.IsValid();
+				}
+			}
+
+			TBitArray<TInlineAllocator<8>> CulledBuffers(true, Graph->Buffers.Num());
+
+			if (ShowBuffers())
+			{
+				for (uint32 BufferIndex = 0, BufferCount = Graph->Buffers.Num(); BufferIndex < BufferCount; ++BufferIndex)
+				{
+					FBufferPacket& Buffer = Graph->Buffers[BufferIndex];
+
+					const bool bCulled = Buffer.bCulled || !CheckFilterName(Buffer.Name) || !CheckFilterSize(Buffer.SizeInBytes);
+					CulledBuffers[BufferIndex] = bCulled;
+
+					if (bCulled)
 					{
 						continue;
 					}
 
+					FResourceEntry& Entry = Resources.AddDefaulted_GetRef();
+					Entry.StartTime = Buffer.StartTime;
+					Entry.EndTime = Buffer.EndTime;
+					Entry.SizeInBytes = Buffer.SizeInBytes;
+					Entry.Index = BufferIndex;
+					Entry.Order = Buffer.Order;
+					Entry.Type = ERDGParentResourceType::Buffer;
+					Entry.bHasPreviousOwner = Buffer.PrevousOwnerHandle.IsValid();
+				}
+			}
+
+			switch (ResourceSort)
+			{
+			case EResourceSort::LargestSize:
+				Resources.StableSort([](const FResourceEntry& LHS, const FResourceEntry& RHS)
+				{
+					return LHS.SizeInBytes > RHS.SizeInBytes;
+				});
+				break;
+			case EResourceSort::SmallestSize:
+				Resources.StableSort([](const FResourceEntry& LHS, const FResourceEntry& RHS)
+				{
+					return LHS.SizeInBytes < RHS.SizeInBytes;
+				});
+				break;
+			case EResourceSort::StartOfLifetime:
+				Resources.StableSort([](const FResourceEntry& LHS, const FResourceEntry& RHS)
+				{
+					return LHS.StartTime < RHS.StartTime;
+				});
+				break;
+			case EResourceSort::EndOfLifetime:
+				Resources.StableSort([](const FResourceEntry& LHS, const FResourceEntry& RHS)
+				{
+					return LHS.EndTime > RHS.EndTime;
+				});
+				break;
+			default:
+				Resources.StableSort([](const FResourceEntry& LHS, const FResourceEntry& RHS)
+				{
+					return LHS.Order < RHS.Order;
+				});
+				break;
+			}
+
+			uint64 MaxSizeInBytes = 0;
+
+			for (FResourceEntry& Entry : Resources)
+			{
+				if (!Entry.bHasPreviousOwner)
+				{
+					auto& Array = Entry.Type == ERDGParentResourceType::Texture ? TextureIndexToDepth : BufferIndexToDepth;
+					Array[Entry.Index] = DepthOffset++;
+				}
+
+				MaxSizeInBytes = FMath::Max(MaxSizeInBytes, Entry.SizeInBytes);
+			}
+
+			if (ShowTextures())
+			{
+				for (uint32 TextureIndex = 0, TextureCount = Graph->Textures.Num(); TextureIndex < TextureCount; ++TextureIndex)
+				{
+					if (CulledTextures[TextureIndex])
+					{
+						continue;
+					}
+
+					FTexturePacket& Texture = Graph->Textures[TextureIndex];
+
+					if (Texture.PrevousOwnerHandle.IsValid())
+					{
+						TextureIndexToDepth[TextureIndex] = TextureIndexToDepth[Texture.PrevousOwnerHandle.GetIndex()];
+					}
+
+					const uint32 Depth = TextureIndexToDepth[TextureIndex];
 					const double StartTime = Texture.StartTime + SinglePixelTimeMargin;
-					const double EndTime   = Texture.EndTime;
-					const uint32 Depth     = DepthOffset + Texture.Depth;
+					const double EndTime = Texture.EndTime;
+					const uint32 Color = GetTextureColor(Texture, MaxSizeInBytes);
 
-					const FVisibleTexture VisibleTexture(Viewport, Texture, StartTime, EndTime, Depth, kTextureColor);
+					const FVisibleTexture VisibleTexture(Viewport, Texture, StartTime, EndTime, Depth, Color);
 					AddEvent(Builder, VisibleTexture);
 					VisibleGraph.AddTexture(VisibleTexture);
-
-					TextureDepthOffsetMax = FMath::Max(Depth, TextureDepthOffsetMax);
 				}
-
-				DepthOffset = FMath::Max(DepthOffset, TextureDepthOffsetMax + 1);
 			}
 
 			if (ShowBuffers())
 			{
-				uint32 BufferDepthOffsetMax = 0;
-
 				for (uint32 BufferIndex = 0, BufferCount = Graph->Buffers.Num(); BufferIndex < BufferCount; ++BufferIndex)
 				{
-					FBufferPacket& Buffer = Graph->Buffers[BufferIndex];
-					if (Buffer.bCulled)
+					if (CulledBuffers[BufferIndex])
 					{
 						continue;
 					}
 
+					FBufferPacket& Buffer = Graph->Buffers[BufferIndex];
+
+					if (Buffer.PrevousOwnerHandle.IsValid())
+					{
+						BufferIndexToDepth[BufferIndex] = BufferIndexToDepth[Buffer.PrevousOwnerHandle.GetIndex()];
+					}
+
+					const uint32 Depth = BufferIndexToDepth[BufferIndex];
 					const double StartTime = Buffer.StartTime + SinglePixelTimeMargin;
-					const double EndTime   = Buffer.EndTime;
-					const uint32 Depth     = DepthOffset + Buffer.Depth;
+					const double EndTime = Buffer.EndTime;
+					const uint32 Color = GetBufferColor(Buffer, MaxSizeInBytes);
 
-					const FVisibleBuffer VisibleBuffer(Viewport, Buffer, StartTime, EndTime, Depth, kBufferColor);
+					const FVisibleBuffer VisibleBuffer(Viewport, Buffer, StartTime, EndTime, Depth, Color);
 					AddEvent(Builder, VisibleBuffer);
 					VisibleGraph.AddBuffer(VisibleBuffer);
-
-					BufferDepthOffsetMax = FMath::Max(Depth, BufferDepthOffsetMax);
 				}
-
-				DepthOffset = FMath::Max(DepthOffset, BufferDepthOffsetMax + 1);
 			}
 
-			check(VisibleGraph.Passes.Num() == VisibleGraph.GetPacket().PassCount);
-
 			return TraceServices::EEventEnumerate::Continue;
 		});
 	}
@@ -800,8 +962,10 @@ void FRenderGraphTrack::BuildFilteredDrawState(ITimingEventsTrackDrawStateBuilde
 			for (FRDGTextureHandle TextureHandle : Pass.Textures)
 			{
 				const FTexturePacket& Texture = *Graph.GetTexture(TextureHandle);
-				const FVisibleTexture& VisibleTexture = VisibleGraph->GetVisibleTexture(Texture);
-				VisibleItems.Add(&VisibleTexture);
+				if (const FVisibleTexture* VisibleTexture = VisibleGraph->GetVisibleTexture(Texture))
+				{
+					VisibleItems.Add(VisibleTexture);
+				}
 			}
 		}
 
@@ -810,8 +974,10 @@ void FRenderGraphTrack::BuildFilteredDrawState(ITimingEventsTrackDrawStateBuilde
 			for (FRDGBufferHandle BufferHandle : Pass.Buffers)
 			{
 				const FBufferPacket& Buffer = *Graph.GetBuffer(BufferHandle);
-				const FVisibleBuffer& VisibleBuffer = VisibleGraph->GetVisibleBuffer(Buffer);
-				VisibleItems.Add(&VisibleBuffer);
+				if (const FVisibleBuffer* VisibleBuffer = VisibleGraph->GetVisibleBuffer(Buffer))
+				{
+					VisibleItems.Add(VisibleBuffer);
+				}
 			}
 		}
 
@@ -979,7 +1145,7 @@ void FRenderGraphTrack::BuildContextMenu(FMenuBuilder& MenuBuilder)
 {
 	Super::BuildContextMenu(MenuBuilder);
 
-	MenuBuilder.BeginSection("Layout", LOCTEXT("TrackLayoutMenuHeader", "Track Layout"));
+	MenuBuilder.BeginSection("Show", LOCTEXT("ShowMenuHeader", "Track Show Flags"));
 	{
 		MenuBuilder.AddMenuEntry
 		(
@@ -1036,6 +1202,182 @@ void FRenderGraphTrack::BuildContextMenu(FMenuBuilder& MenuBuilder)
 		);
 	}
 	MenuBuilder.EndSection();
+	
+	MenuBuilder.BeginSection("Sort", LOCTEXT("SortMenuHeader", "Track Sort By"));
+	{
+		MenuBuilder.AddMenuEntry
+		(
+			LOCTEXT("SortCreation", "Creation"),
+			LOCTEXT("SortCreation_Tooltip", "Resources created earlier in the graph builder are ordered first."),
+			FSlateIcon(),
+			FUIAction(
+				FExecuteAction::CreateLambda([this]()
+				{
+					ResourceSort = EResourceSort::Creation;
+					SetDirtyFlag();
+				}),
+				FCanExecuteAction(),
+				FIsActionChecked::CreateLambda([this]() { return ResourceSort == EResourceSort::Creation; })
+			),
+			NAME_None,
+			EUserInterfaceActionType::RadioButton
+		);
+
+		MenuBuilder.AddMenuEntry
+		(
+			LOCTEXT("SortLargestSize", "Largest Size"),
+			LOCTEXT("SortLargestSize_Tooltip", "Resources with larger allocations are ordered first."),
+			FSlateIcon(),
+			FUIAction(
+				FExecuteAction::CreateLambda([this]()
+				{
+					ResourceSort = EResourceSort::LargestSize;
+					SetDirtyFlag();
+				}),
+				FCanExecuteAction(),
+				FIsActionChecked::CreateLambda([this]() { return ResourceSort == EResourceSort::LargestSize; })
+			),
+			NAME_None,
+			EUserInterfaceActionType::RadioButton
+		);
+
+		MenuBuilder.AddMenuEntry
+		(
+			LOCTEXT("SortSmallestSize", "Smallest Size"),
+			LOCTEXT("SortSmallestSize_Tooltip", "Resources with smaller allocations are ordered first."),
+			FSlateIcon(),
+			FUIAction(
+				FExecuteAction::CreateLambda([this]()
+				{
+					ResourceSort = EResourceSort::SmallestSize;
+					SetDirtyFlag();
+				}),
+				FCanExecuteAction(),
+				FIsActionChecked::CreateLambda([this]() { return ResourceSort == EResourceSort::SmallestSize; })
+			),
+			NAME_None,
+			EUserInterfaceActionType::RadioButton
+		);
+
+		MenuBuilder.AddMenuEntry
+		(
+			LOCTEXT("SortStartOfLifetime", "Start Of Lifetime"),
+			LOCTEXT("SortStartOfLifetime_Tooltip", "Resources with earlier starting lifetimes are ordered first."),
+			FSlateIcon(),
+			FUIAction(
+				FExecuteAction::CreateLambda([this]()
+				{
+					ResourceSort = EResourceSort::StartOfLifetime;
+					SetDirtyFlag();
+				}),
+				FCanExecuteAction(),
+				FIsActionChecked::CreateLambda([this]() { return ResourceSort == EResourceSort::StartOfLifetime; })
+			),
+			NAME_None,
+			EUserInterfaceActionType::RadioButton
+		);
+		
+		MenuBuilder.AddMenuEntry
+		(
+			LOCTEXT("SortEndOfLifetime", "End Of Lifetime"),
+			LOCTEXT("SortEndOfLifetime_Tooltip", "Resources with later ending lifetimes are ordered first."),
+			FSlateIcon(),
+			FUIAction(
+				FExecuteAction::CreateLambda([this]()
+				{
+					ResourceSort = EResourceSort::EndOfLifetime;
+					SetDirtyFlag();
+				}),
+				FCanExecuteAction(),
+				FIsActionChecked::CreateLambda([this]() { return ResourceSort == EResourceSort::EndOfLifetime; })
+			),
+			NAME_None,
+			EUserInterfaceActionType::RadioButton
+		);
+	}
+	MenuBuilder.EndSection();
+
+	MenuBuilder.BeginSection("Color", LOCTEXT("ColorMenuHeader", "Track Resource Coloration"));
+	{
+		MenuBuilder.AddMenuEntry
+		(
+			LOCTEXT("ColorType", "By Type"),
+			LOCTEXT("ColorType_Tooltip", "Each type of resource has a unique color."),
+			FSlateIcon(),
+			FUIAction(
+				FExecuteAction::CreateLambda([this]()
+				{
+					ResourceColor = EResourceColor::Type;
+					SetDirtyFlag();
+				}),
+				FCanExecuteAction(),
+				FIsActionChecked::CreateLambda([this]() { return ResourceColor == EResourceColor::Type; })
+			),
+			NAME_None,
+			EUserInterfaceActionType::RadioButton
+		);
+
+		MenuBuilder.AddMenuEntry
+		(
+			LOCTEXT("ColorSize", "By Size"),
+			LOCTEXT("ColorSize_Tooltip", "Larger resources are more brightly colored."),
+			FSlateIcon(),
+			FUIAction(
+				FExecuteAction::CreateLambda([this]()
+				{
+					ResourceColor = EResourceColor::Size;
+					SetDirtyFlag();
+				}),
+				FCanExecuteAction(),
+				FIsActionChecked::CreateLambda([this]() { return ResourceColor == EResourceColor::Size; })
+			),
+			NAME_None,
+			EUserInterfaceActionType::RadioButton
+		);
+	}
+	MenuBuilder.EndSection();
+
+	MenuBuilder.BeginSection("FilterText", LOCTEXT("FilterTextHeader", "Track Resource Filter"));
+	{
+		MenuBuilder.AddWidget(
+			SNew(SVerticalBox)
+			+ SVerticalBox::Slot()
+			.AutoHeight()
+			.Padding(5.f)
+			[
+				// Search box allows for filtering
+				SNew(SSearchBox)
+				.InitialText(FText::FromString(FilterText))
+				.HintText(LOCTEXT("SearchHint", "Filter By Name"))
+				.OnTextChanged_Lambda([this](const FText& InText) { FilterText = InText.ToString(); SetDirtyFlag(); })
+			]
+			+ SVerticalBox::Slot()
+			.AutoHeight()
+			.Padding(5.f)
+			[
+				SNew(SHorizontalBox)
+				+ SHorizontalBox::Slot()
+				.VAlign(VAlign_Center)
+				[
+					SNew(STextBlock)
+					.Text(LOCTEXT("SizeThreshold", "Filter By Size (MB)"))
+				]
+				+ SHorizontalBox::Slot()
+				.AutoWidth()
+				.VAlign(VAlign_Center)
+				[
+					SNew(SSpinBox<float>)
+					.MinValue(0)
+					.MaxValue(1024)
+					.Value(FilterSize)
+					.MaxFractionalDigits(3)
+					.MinDesiredWidth(60)
+					.OnValueCommitted_Lambda([this](float InValue, ETextCommit::Type) { FilterSize = InValue; SetDirtyFlag(); })
+				]
+			],
+			FText::GetEmpty(), true);
+	}
+	MenuBuilder.EndSection();
 }
 
 void FRenderGraphTrack::InitTooltip(FTooltipDrawState& Tooltip, const ITimingEvent& InTooltipEvent) const
diff --git a/Engine/Plugins/RenderGraphInsights/Source/RenderGraphInsights/Private/RenderGraphTrack.h b/Engine/Plugins/RenderGraphInsights/Source/RenderGraphInsights/Private/RenderGraphTrack.h
--- a/Engine/Plugins/RenderGraphInsights/Source/RenderGraphInsights/Private/RenderGraphTrack.h
+++ b/Engine/Plugins/RenderGraphInsights/Source/RenderGraphInsights/Private/RenderGraphTrack.h
@@ -160,14 +160,24 @@ public:
 		return Passes[Pass.VisibleIndex];
 	}
 
-	const FVisibleTexture& GetVisibleTexture(const FTexturePacket& Texture) const
+	const FVisibleTexture* GetVisibleTexture(const FTexturePacket& Texture) const
 	{
-		return Textures[Texture.VisibleIndex];
+		if (Texture.VisibleIndex == kInvalidVisibleIndex)
+		{
+			return nullptr;
+		}
+
+		return &Textures[Texture.VisibleIndex];
 	}
 
-	const FVisibleBuffer& GetVisibleBuffer(const FBufferPacket& Buffer) const
+	const FVisibleBuffer* GetVisibleBuffer(const FBufferPacket& Buffer) const
 	{
-		return Buffers[Buffer.VisibleIndex];
+		if (Buffer.VisibleIndex == kInvalidVisibleIndex)
+		{
+			return nullptr;
+		}
+
+		return &Buffers[Buffer.VisibleIndex];
 	}
 
 	TArray<FVisibleScope> Scopes;
@@ -241,6 +251,9 @@ private:
 	const FPacketFilter* PacketFilter{};
 	const FPassPacket* SelectedPass{};
 
+	uint32 GetTextureColor(const FTexturePacket& Texture, uint64 MaxSizeInBytes) const;
+	uint32 GetBufferColor(const FBufferPacket& Buffer, uint64 MaxSizeInBytes) const;
+
 	bool ShowTextures() const
 	{
 		return EnumHasAnyFlags(ResourceShow, EResourceShow::Textures);
@@ -259,9 +272,26 @@ private:
 	};
 	FRIEND_ENUM_CLASS_FLAGS(FRenderGraphTrack::EResourceShow);
 
-	static void AddEvent(ITimingEventsTrackDrawStateBuilder& Builder, const FVisibleItem& Item);
+	enum class EResourceSort
+	{
+		Creation,
+		LargestSize,
+		SmallestSize,
+		StartOfLifetime,
+		EndOfLifetime
+	};
+
+	enum class EResourceColor
+	{
+		Type,
+		Size
+	};
 
 	EResourceShow ResourceShow = EResourceShow::All;
+	EResourceSort ResourceSort = EResourceSort::Creation;
+	EResourceColor ResourceColor = EResourceColor::Type;
+	FString FilterText;
+	float FilterSize{};
 
 	FTooltipDrawState SelectedTooltipState;
 };
diff --git a/Engine/Plugins/RenderGraphInsights/Source/RenderGraphInsights/RenderGraphInsights.Build.cs b/Engine/Plugins/RenderGraphInsights/Source/RenderGraphInsights/RenderGraphInsights.Build.cs
--- a/Engine/Plugins/RenderGraphInsights/Source/RenderGraphInsights/RenderGraphInsights.Build.cs
+++ b/Engine/Plugins/RenderGraphInsights/Source/RenderGraphInsights/RenderGraphInsights.Build.cs
@@ -18,7 +18,8 @@ namespace UnrealBuildTool.Rules
 				"TraceLog",
 				"TraceAnalysis",
 				"TraceServices",
-				"TraceInsights"
+				"TraceInsights",
+				"InputCore"
 			});
 		}
 	}
diff --git a/Engine/Source/Runtime/RenderCore/Private/RenderGraphBuilder.cpp b/Engine/Source/Runtime/RenderCore/Private/RenderGraphBuilder.cpp
--- a/Engine/Source/Runtime/RenderCore/Private/RenderGraphBuilder.cpp
+++ b/Engine/Source/Runtime/RenderCore/Private/RenderGraphBuilder.cpp
@@ -526,6 +526,7 @@ FRDGTextureRef FRDGBuilder::RegisterExternalTexture(
 	ExternalTextures.Add(Texture->GetRHIUnchecked(), Texture);
 
 	IF_RDG_ENABLE_DEBUG(UserValidation.ValidateRegisterExternalTexture(Texture));
+	IF_RDG_ENABLE_TRACE(Trace.AddResource(Texture));
 	return Texture;
 }
 
@@ -580,6 +581,7 @@ FRDGBufferRef FRDGBuilder::RegisterExternalBuffer(
 	ExternalBuffers.Add(ExternalPooledBuffer, Buffer);
 
 	IF_RDG_ENABLE_DEBUG(UserValidation.ValidateRegisterExternalBuffer(Buffer));
+	IF_RDG_ENABLE_TRACE(Trace.AddResource(Buffer));
 	return Buffer;
 }
 
@@ -1790,10 +1792,7 @@ void FRDGBuilder::CollectPassBarriers(FRDGPassHandle PassHandle, FRDGPassHandle&
 		AddTransition(PassHandle, Texture, TexturePair.Value.MergeState, LastUntrackedPassHandle);
 		Texture->bCulled = false;
 
-	#if RDG_ENABLE_TRACE
-		Pass->Textures.Add(Texture->Handle);
-		Texture->Passes.Add(PassHandle);
-	#endif
+		IF_RDG_ENABLE_TRACE(Trace.AddTexturePassDependency(Texture, Pass));
 	}
 
 	for (const auto& BufferPair : Pass->BufferStates)
@@ -1802,10 +1801,7 @@ void FRDGBuilder::CollectPassBarriers(FRDGPassHandle PassHandle, FRDGPassHandle&
 		AddTransition(PassHandle, Buffer, *BufferPair.Value.MergeState, LastUntrackedPassHandle);
 		Buffer->bCulled = false;
 
-	#if RDG_ENABLE_TRACE
-		Pass->Buffers.Add(Buffer->Handle);
-		Buffer->Passes.Add(PassHandle);
-	#endif
+		IF_RDG_ENABLE_TRACE(Trace.AddBufferPassDependency(Buffer, Pass));
 	}
 }
 
diff --git a/Engine/Source/Runtime/RenderCore/Private/RenderGraphTrace.cpp b/Engine/Source/Runtime/RenderCore/Private/RenderGraphTrace.cpp
--- a/Engine/Source/Runtime/RenderCore/Private/RenderGraphTrace.cpp
+++ b/Engine/Source/Runtime/RenderCore/Private/RenderGraphTrace.cpp
@@ -44,6 +44,7 @@ UE_TRACE_EVENT_BEGIN(RDGTrace, BufferMessage)
 	UE_TRACE_EVENT_FIELD(uint32, NumElements)
 	UE_TRACE_EVENT_FIELD(uint16, Handle)
 	UE_TRACE_EVENT_FIELD(uint16, NextOwnerHandle)
+	UE_TRACE_EVENT_FIELD(uint16, Order)
 	UE_TRACE_EVENT_FIELD(uint16[], Passes)
 	UE_TRACE_EVENT_FIELD(bool, IsExternal)
 	UE_TRACE_EVENT_FIELD(bool, IsExtracted)
@@ -56,6 +57,7 @@ UE_TRACE_EVENT_BEGIN(RDGTrace, TextureMessage)
 	UE_TRACE_EVENT_FIELD(uint64, EndCycles)
 	UE_TRACE_EVENT_FIELD(uint16, Handle)
 	UE_TRACE_EVENT_FIELD(uint16, NextOwnerHandle)
+	UE_TRACE_EVENT_FIELD(uint16, Order)
 	UE_TRACE_EVENT_FIELD(uint16[], Passes)
 	UE_TRACE_EVENT_FIELD(uint64, SizeInBytes)
 	UE_TRACE_EVENT_FIELD(uint64, CreateFlags)
@@ -131,8 +133,8 @@ void FRDGTrace::OutputGraphEnd(const FRDGBuilder& GraphBuilder)
 			<< PassMessage.Handle(Handle.GetIndex())
 			<< PassMessage.GraphicsForkPass(Pass->GetGraphicsForkPass().GetIndexUnchecked())
 			<< PassMessage.GraphicsJoinPass(Pass->GetGraphicsJoinPass().GetIndexUnchecked())
-			<< PassMessage.Textures((const uint16*)Pass->Textures.GetData(), (uint16)Pass->Textures.Num())
-			<< PassMessage.Buffers((const uint16*)Pass->Buffers.GetData(), (uint16)Pass->Buffers.Num())
+			<< PassMessage.Textures((const uint16*)Pass->TraceTextures.GetData(), (uint16)Pass->TraceTextures.Num())
+			<< PassMessage.Buffers((const uint16*)Pass->TraceBuffers.GetData(), (uint16)Pass->TraceBuffers.Num())
 			<< PassMessage.Flags(uint16(Pass->GetFlags()))
 			<< PassMessage.Pipeline(uint16(Pass->GetPipeline()))
 			<< PassMessage.IsCulled(GraphBuilder.PassesToCull[Handle])
@@ -219,7 +221,8 @@ void FRDGTrace::OutputGraphEnd(const FRDGBuilder& GraphBuilder)
 			<< TextureMessage.Name(Texture->Name, uint16(FCString::Strlen(Texture->Name)))
 			<< TextureMessage.Handle(Handle.GetIndex())
 			<< TextureMessage.NextOwnerHandle(Texture->NextOwner.GetIndexUnchecked())
-			<< TextureMessage.Passes((const uint16*)Texture->Passes.GetData(), (uint16)Texture->Passes.Num())
+			<< TextureMessage.Order(Texture->TraceOrder)
+			<< TextureMessage.Passes((const uint16*)Texture->TracePasses.GetData(), (uint16)Texture->TracePasses.Num())
 			<< TextureMessage.SizeInBytes(SizeInBytes)
 			<< TextureMessage.CreateFlags(uint32(Texture->Desc.Flags))
 			<< TextureMessage.Dimension(uint32(Texture->Desc.Dimension))
@@ -243,7 +246,8 @@ void FRDGTrace::OutputGraphEnd(const FRDGBuilder& GraphBuilder)
 			<< BufferMessage.Name(Buffer->Name, uint16(FCString::Strlen(Buffer->Name)))
 			<< BufferMessage.Handle(Buffer->Handle.GetIndex())
 			<< BufferMessage.NextOwnerHandle(Buffer->NextOwner.GetIndexUnchecked())
-			<< BufferMessage.Passes((const uint16*)Buffer->Passes.GetData(), (uint16)Buffer->Passes.Num())
+			<< BufferMessage.Order(Buffer->TraceOrder)
+			<< BufferMessage.Passes((const uint16*)Buffer->TracePasses.GetData(), (uint16)Buffer->TracePasses.Num())
 			<< BufferMessage.UsageFlags(uint32(Buffer->Desc.Usage))
 			<< BufferMessage.BytesPerElement(Buffer->Desc.BytesPerElement)
 			<< BufferMessage.NumElements(Buffer->Desc.NumElements)
@@ -255,4 +259,21 @@ void FRDGTrace::OutputGraphEnd(const FRDGBuilder& GraphBuilder)
 	UE_TRACE_LOG(RDGTrace, GraphEndMessage, RDGChannel);
 }
 
+void FRDGTrace::AddResource(FRDGParentResource* Resource)
+{
+	Resource->TraceOrder = ResourceOrder++;
+}
+
+void FRDGTrace::AddTexturePassDependency(FRDGTexture* Texture, FRDGPass* Pass)
+{
+	Pass->TraceTextures.Add(Texture->Handle);
+	Texture->TracePasses.Add(Pass->Handle);
+}
+
+void FRDGTrace::AddBufferPassDependency(FRDGBuffer* Buffer, FRDGPass* Pass)
+{
+	Pass->TraceBuffers.Add(Buffer->Handle);
+	Buffer->TracePasses.Add(Pass->Handle);
+}
+
 #endif
\ No newline at end of file
diff --git a/Engine/Source/Runtime/RenderCore/Public/RenderGraphBuilder.inl b/Engine/Source/Runtime/RenderCore/Public/RenderGraphBuilder.inl
--- a/Engine/Source/Runtime/RenderCore/Public/RenderGraphBuilder.inl
+++ b/Engine/Source/Runtime/RenderCore/Public/RenderGraphBuilder.inl
@@ -28,6 +28,7 @@ inline FRDGTextureRef FRDGBuilder::CreateTexture(
 	IF_RDG_ENABLE_DEBUG(UserValidation.ValidateCreateTexture(Desc, Name, Flags));
 	FRDGTextureRef Texture = Textures.Allocate(Allocator, Name, Desc, Flags, ERenderTargetTexture::ShaderResource);
 	IF_RDG_ENABLE_DEBUG(UserValidation.ValidateCreateTexture(Texture));
+	IF_RDG_ENABLE_TRACE(Trace.AddResource(Texture));
 	return Texture;
 }
 
@@ -39,6 +40,7 @@ inline FRDGBufferRef FRDGBuilder::CreateBuffer(
 	IF_RDG_ENABLE_DEBUG(UserValidation.ValidateCreateBuffer(Desc, Name, Flags));
 	FRDGBufferRef Buffer = Buffers.Allocate(Allocator, Name, Desc, Flags);
 	IF_RDG_ENABLE_DEBUG(UserValidation.ValidateCreateBuffer(Buffer));
+	IF_RDG_ENABLE_TRACE(Trace.AddResource(Buffer));
 	return Buffer;
 }
 
diff --git a/Engine/Source/Runtime/RenderCore/Public/RenderGraphDefinitions.h b/Engine/Source/Runtime/RenderCore/Public/RenderGraphDefinitions.h
--- a/Engine/Source/Runtime/RenderCore/Public/RenderGraphDefinitions.h
+++ b/Engine/Source/Runtime/RenderCore/Public/RenderGraphDefinitions.h
@@ -502,6 +502,9 @@ using FRDGBufferHandle = TRDGHandle<FRDGBuffer, uint16>;
 using FRDGBufferRegistry = TRDGHandleRegistry<FRDGBufferHandle>;
 using FRDGBufferBitArray = TRDGHandleBitArray<FRDGBufferHandle>;
 
+class FRDGPooledTexture;
+class FRDGPooledBuffer;
+
 template <typename TUniformStruct> class TRDGUniformBuffer;
 template <typename TUniformStruct> using TRDGUniformBufferRef = TRDGUniformBuffer<TUniformStruct>*;
 
diff --git a/Engine/Source/Runtime/RenderCore/Public/RenderGraphPass.h b/Engine/Source/Runtime/RenderCore/Public/RenderGraphPass.h
--- a/Engine/Source/Runtime/RenderCore/Public/RenderGraphPass.h
+++ b/Engine/Source/Runtime/RenderCore/Public/RenderGraphPass.h
@@ -384,8 +384,8 @@ private:
 #endif
 
 #if RDG_ENABLE_TRACE
-	TArray<FRDGTextureHandle, FRDGArrayAllocator> Textures;
-	TArray<FRDGBufferHandle, FRDGArrayAllocator> Buffers;
+	TArray<FRDGTextureHandle, FRDGArrayAllocator> TraceTextures;
+	TArray<FRDGBufferHandle, FRDGArrayAllocator> TraceBuffers;
 #endif
 
 	friend FRDGBuilder;
diff --git a/Engine/Source/Runtime/RenderCore/Public/RenderGraphResources.h b/Engine/Source/Runtime/RenderCore/Public/RenderGraphResources.h
--- a/Engine/Source/Runtime/RenderCore/Public/RenderGraphResources.h
+++ b/Engine/Source/Runtime/RenderCore/Public/RenderGraphResources.h
@@ -274,7 +274,8 @@ private:
 	FRDGPassHandle LastPass;
 
 #if RDG_ENABLE_TRACE
-	TArray<FRDGPassHandle, FRDGArrayAllocator> Passes;
+	uint16 TraceOrder = 0;
+	TArray<FRDGPassHandle, FRDGArrayAllocator> TracePasses;
 #endif
 
 #if RDG_ENABLE_DEBUG
diff --git a/Engine/Source/Runtime/RenderCore/Public/RenderGraphTrace.h b/Engine/Source/Runtime/RenderCore/Public/RenderGraphTrace.h
--- a/Engine/Source/Runtime/RenderCore/Public/RenderGraphTrace.h
+++ b/Engine/Source/Runtime/RenderCore/Public/RenderGraphTrace.h
@@ -14,8 +14,13 @@ public:
 	void OutputGraphBegin();
 	void OutputGraphEnd(const FRDGBuilder& GraphBuilder);
 
+	void AddResource(FRDGParentResource* Resource);
+	void AddTexturePassDependency(FRDGTexture* Texture, FRDGPass* Pass);
+	void AddBufferPassDependency(FRDGBuffer* Buffer, FRDGPass* Pass);
+
 private:
 	uint64 GraphStartCycles{};
+	uint32 ResourceOrder{};
 };
 
 #endif
\ No newline at end of file
