Added tooltip to show number of passes in scope for RDG insights.

#rnx

diff --git a/Engine/Plugins/RenderGraphInsights/Source/RenderGraphInsights/Private/RenderGraphTrack.cpp b/Engine/Plugins/RenderGraphInsights/Source/RenderGraphInsights/Private/RenderGraphTrack.cpp
--- a/Engine/Plugins/RenderGraphInsights/Source/RenderGraphInsights/Private/RenderGraphTrack.cpp
+++ b/Engine/Plugins/RenderGraphInsights/Source/RenderGraphInsights/Private/RenderGraphTrack.cpp
@@ -1404,6 +1404,9 @@ void FRenderGraphTrack::InitTooltip(FTooltipDrawState& Tooltip, const ITimingEve
 			const auto& TooltipEvent = InTooltipEvent.As<FVisibleScopeEvent>();
 			const FScopePacket& Scope = TooltipEvent.GetPacket();
 			Tooltip.AddTitle(Scope.Name);
+
+			const uint32 PassCount = Scope.LastPass.GetIndex() - Scope.FirstPass.GetIndex() + 1;
+			Tooltip.AddNameValueTextLine(TEXT("Passes:"), FString::Printf(TEXT("%d"), PassCount));
 		}
 		else if (InTooltipEvent.Is<FVisiblePassEvent>())
 		{
