Improved RDG insights show flags to allow filtering of external / transient / pooled resources from view.

diff --git a/Engine/Plugins/RenderGraphInsights/Source/RenderGraphInsights/Private/RenderGraphTrack.cpp b/Engine/Plugins/RenderGraphInsights/Source/RenderGraphInsights/Private/RenderGraphTrack.cpp
--- a/Engine/Plugins/RenderGraphInsights/Source/RenderGraphInsights/Private/RenderGraphTrack.cpp
+++ b/Engine/Plugins/RenderGraphInsights/Source/RenderGraphInsights/Private/RenderGraphTrack.cpp
@@ -767,14 +767,39 @@ void FRenderGraphTrack::BuildDrawState(ITimingEventsTrackDrawStateBuilder& Build
 
 			TBitArray<TInlineAllocator<8>> CulledTextures(true, Graph->Textures.Num());
 
-			const auto CheckFilterName = [this](const FString& InName)
+			const auto IsPacketCulled = [&](const FResourcePacket& Packet)
 			{
-				return FilterText.IsEmpty() || InName.Contains(FilterText);
-			};
+				if (Packet.bCulled)
+				{
+					return true;
+				}
 
-			const auto CheckFilterSize = [this](uint64 InSizeInBytes)
-			{
-				return FilterSize <= 0.0f || InSizeInBytes >= uint64(FilterSize * 1024.0f * 1024.0f);
+				if (Packet.bTransient && !ShowTransient())
+				{
+					return true;
+				}
+
+				if (Packet.bExternal && !ShowExternal())
+				{
+					return true;
+				}
+
+				if (!Packet.bTransient && !Packet.bExternal && !ShowPooled())
+				{
+					return true;
+				}
+
+				if (FilterSize > 0.0f && Packet.SizeInBytes < uint64(FilterSize * 1024.0f * 1024.0f))
+				{
+					return true;
+				}
+
+				if (!FilterText.IsEmpty() && !Packet.Name.Contains(FilterText))
+				{
+					return true;
+				}
+
+				return false;
 			};
 
 			if (ShowTextures())
@@ -783,7 +808,7 @@ void FRenderGraphTrack::BuildDrawState(ITimingEventsTrackDrawStateBuilder& Build
 				{
 					FTexturePacket& Texture = Graph->Textures[TextureIndex];
 
-					const bool bCulled = Texture.bCulled || !CheckFilterName(Texture.Name) || !CheckFilterSize(Texture.SizeInBytes);
+					const bool bCulled = IsPacketCulled(Texture);
 					CulledTextures[TextureIndex] = bCulled;
 
 					if (bCulled)
@@ -810,7 +835,7 @@ void FRenderGraphTrack::BuildDrawState(ITimingEventsTrackDrawStateBuilder& Build
 				{
 					FBufferPacket& Buffer = Graph->Buffers[BufferIndex];
 
-					const bool bCulled = Buffer.bCulled || !CheckFilterName(Buffer.Name) || !CheckFilterSize(Buffer.SizeInBytes);
+					const bool bCulled = IsPacketCulled(Buffer);
 					CulledBuffers[BufferIndex] = bCulled;
 
 					if (bCulled)
@@ -1167,60 +1192,96 @@ void FRenderGraphTrack::BuildContextMenu(FMenuBuilder& MenuBuilder)
 	{
 		MenuBuilder.AddMenuEntry
 		(
-			LOCTEXT("ShowAlls", "Show All"),
-			LOCTEXT("ShowAlls_Tooltip", "Show All resources in the lifetime view."),
+			LOCTEXT("ShowTextures", "Show Textures"),
+			LOCTEXT("ShowTextures_Tooltip", "Show Texture resources in the lifetime view."),
 			FSlateIcon(),
 			FUIAction(
 				FExecuteAction::CreateLambda([this]()
 				{
-					ResourceShow = EResourceShow::All;
+					ResourceShow ^= EResourceShow::Textures;
 					SetDirtyFlag();
 				}),
 				FCanExecuteAction(),
-				FIsActionChecked::CreateLambda([this]() { return ResourceShow == EResourceShow::All; })
+				FIsActionChecked::CreateLambda([this]() { return EnumHasAnyFlags(ResourceShow, EResourceShow::Textures); })
 			),
 			NAME_None,
-			EUserInterfaceActionType::RadioButton
+			EUserInterfaceActionType::ToggleButton
 		);
 
 		MenuBuilder.AddMenuEntry
 		(
-			LOCTEXT("ShowTextures", "Show Textures"),
-			LOCTEXT("ShowTextures_Tooltip", "Show Texture resources in the lifetime view."),
+			LOCTEXT("ShowBuffers", "Show Buffers"),
+			LOCTEXT("ShowBuffers_Tooltip", "Show Buffer resources in the lifetime view."),
 			FSlateIcon(),
 			FUIAction(
 				FExecuteAction::CreateLambda([this]()
 				{
-					ResourceShow = EResourceShow::Textures;
+					ResourceShow ^= EResourceShow::Buffers;
 					SetDirtyFlag();
 				}),
 				FCanExecuteAction(),
-				FIsActionChecked::CreateLambda([this]() { return ResourceShow == EResourceShow::Textures; })
+				FIsActionChecked::CreateLambda([this]() { return EnumHasAnyFlags(ResourceShow, EResourceShow::Buffers); })
 			),
 			NAME_None,
-			EUserInterfaceActionType::RadioButton
+			EUserInterfaceActionType::ToggleButton
 		);
 
 		MenuBuilder.AddMenuEntry
 		(
-			LOCTEXT("ShowBuffers", "Show Buffers"),
-			LOCTEXT("ShowBuffers_Tooltip", "Show Buffer resources in the lifetime view."),
+			LOCTEXT("ShowTransient", "Show Transient"),
+			LOCTEXT("ShowTransient_Tooltip", "Show transient resources in the lifetime view."),
 			FSlateIcon(),
 			FUIAction(
 				FExecuteAction::CreateLambda([this]()
 				{
-					ResourceShow = EResourceShow::Buffers;
+					ResourceShow ^= EResourceShow::Transient;
 					SetDirtyFlag();
 				}),
 				FCanExecuteAction(),
-				FIsActionChecked::CreateLambda([this]() { return ResourceShow == EResourceShow::Buffers; })
+				FIsActionChecked::CreateLambda([this]() { return EnumHasAnyFlags(ResourceShow, EResourceShow::Transient); })
 			),
 			NAME_None,
-			EUserInterfaceActionType::RadioButton
+			EUserInterfaceActionType::ToggleButton
+		);
+		
+		MenuBuilder.AddMenuEntry
+		(
+			LOCTEXT("ShowExternal", "Show External"),
+			LOCTEXT("ShowExternal_Tooltip", "Show external resources in the lifetime view."),
+			FSlateIcon(),
+			FUIAction(
+				FExecuteAction::CreateLambda([this]()
+				{
+					ResourceShow ^= EResourceShow::External;
+					SetDirtyFlag();
+				}),
+				FCanExecuteAction(),
+				FIsActionChecked::CreateLambda([this]() { return EnumHasAnyFlags(ResourceShow, EResourceShow::External); })
+			),
+			NAME_None,
+			EUserInterfaceActionType::ToggleButton
+		);
+		
+		MenuBuilder.AddMenuEntry
+		(
+			LOCTEXT("ShowPooled", "Show Pooled"),
+			LOCTEXT("ShowPooled_Tooltip", "Show pooled resources in the lifetime view."),
+			FSlateIcon(),
+			FUIAction(
+				FExecuteAction::CreateLambda([this]()
+				{
+					ResourceShow ^= EResourceShow::Pooled;
+					SetDirtyFlag();
+				}),
+				FCanExecuteAction(),
+				FIsActionChecked::CreateLambda([this]() { return EnumHasAnyFlags(ResourceShow, EResourceShow::Pooled); })
+			),
+			NAME_None,
+			EUserInterfaceActionType::ToggleButton
 		);
 	}
 	MenuBuilder.EndSection();
-	
+
 	MenuBuilder.BeginSection("Sort", LOCTEXT("SortMenuHeader", "Track Sort By"));
 	{
 		MenuBuilder.AddMenuEntry
diff --git a/Engine/Plugins/RenderGraphInsights/Source/RenderGraphInsights/Private/RenderGraphTrack.h b/Engine/Plugins/RenderGraphInsights/Source/RenderGraphInsights/Private/RenderGraphTrack.h
--- a/Engine/Plugins/RenderGraphInsights/Source/RenderGraphInsights/Private/RenderGraphTrack.h
+++ b/Engine/Plugins/RenderGraphInsights/Source/RenderGraphInsights/Private/RenderGraphTrack.h
@@ -254,21 +254,20 @@ private:
 	uint32 GetTextureColor(const FTexturePacket& Texture, uint64 MaxSizeInBytes) const;
 	uint32 GetBufferColor(const FBufferPacket& Buffer, uint64 MaxSizeInBytes) const;
 
-	bool ShowTextures() const
-	{
-		return EnumHasAnyFlags(ResourceShow, EResourceShow::Textures);
-	}
-
-	bool ShowBuffers() const
-	{
-		return EnumHasAnyFlags(ResourceShow, EResourceShow::Buffers);
-	}
+	bool ShowTextures() const	{ return EnumHasAnyFlags(ResourceShow, EResourceShow::Textures); }
+	bool ShowBuffers() const	{ return EnumHasAnyFlags(ResourceShow, EResourceShow::Buffers); }
+	bool ShowTransient() const	{ return EnumHasAnyFlags(ResourceShow, EResourceShow::Transient); }
+	bool ShowExternal() const	{ return EnumHasAnyFlags(ResourceShow, EResourceShow::External); }
+	bool ShowPooled() const		{ return EnumHasAnyFlags(ResourceShow, EResourceShow::Pooled);}
 
 	enum class EResourceShow
 	{
-		Textures = 1 << 0,
-		Buffers  = 1 << 1,
-		All = Textures | Buffers
+		Textures		= 1 << 0,
+		Buffers			= 1 << 1,
+		Transient		= 1 << 2,
+		External		= 1 << 3,
+		Pooled			= 1 << 4,
+		All				= Textures | Buffers | Transient | External | Pooled
 	};
 	FRIEND_ENUM_CLASS_FLAGS(FRenderGraphTrack::EResourceShow);
 
