RDG insights refactor and transient heap visualizer.

diff --git a/Engine/Plugins/RenderGraphInsights/Source/RenderGraphInsights/Private/RenderGraphProvider.cpp b/Engine/Plugins/RenderGraphInsights/Source/RenderGraphInsights/Private/RenderGraphProvider.cpp
--- a/Engine/Plugins/RenderGraphInsights/Source/RenderGraphInsights/Private/RenderGraphProvider.cpp
+++ b/Engine/Plugins/RenderGraphInsights/Source/RenderGraphInsights/Private/RenderGraphProvider.cpp
@@ -67,6 +67,10 @@ FResourcePacket::FResourcePacket(const UE::Trace::IAnalyzer::FOnEventContext& Co
 	, bCulled(Context.EventData.GetValue<bool>("IsCulled"))
 	, bTransient(Context.EventData.GetValue<bool>("IsTransient"))
 {
+	TransientStats.HeapIndex = Context.EventData.GetValue<uint16>("TransientHeapIndex");
+	TransientStats.HeapOffsetMin = Context.EventData.GetValue<uint64>("TransientHeapOffsetMin");
+	TransientStats.HeapOffsetMax = Context.EventData.GetValue<uint64>("TransientHeapOffsetMax");
+
 	if (Passes.Num())
 	{
 		FirstPass = Passes[0];
@@ -136,6 +140,22 @@ FGraphPacket::FGraphPacket(TraceServices::ILinearAllocator& Allocator, const UE:
 	, PassCount(Context.EventData.GetValue<uint16>("PassCount"))
 {
 	NormalizedPassDuration = (EndTime - StartTime) / double(PassCount);
+
+	TConstArrayView<uint64> WatermarkSizes(Context.EventData.GetArrayView<uint64>("TransientHeapWatermarkSizes"));
+	TConstArrayView<uint64> Capacities(Context.EventData.GetArrayView<uint64>("TransientHeapCapacities"));
+	check(WatermarkSizes.Num() == Capacities.Num());
+
+	uint64 currentOffset = 0;
+
+	for (int32 Index = 0; Index < WatermarkSizes.Num(); ++Index)
+	{
+		auto& Heap = TransientHeapStats.Heaps.AddDefaulted_GetRef();
+		Heap.WatermarkSize = WatermarkSizes[Index];
+		Heap.Capacity = Capacities[Index];
+
+		TransientHeapOffsets.Emplace(currentOffset);
+		currentOffset += Heap.WatermarkSize;
+	}
 }
 
 FRenderGraphProvider::FRenderGraphProvider(TraceServices::IAnalysisSession& InSession)
diff --git a/Engine/Plugins/RenderGraphInsights/Source/RenderGraphInsights/Private/RenderGraphProvider.h b/Engine/Plugins/RenderGraphInsights/Source/RenderGraphInsights/Private/RenderGraphProvider.h
--- a/Engine/Plugins/RenderGraphInsights/Source/RenderGraphInsights/Private/RenderGraphProvider.h
+++ b/Engine/Plugins/RenderGraphInsights/Source/RenderGraphInsights/Private/RenderGraphProvider.h
@@ -76,6 +76,7 @@ public:
 	uint32 Index{};
 	uint32 Order{};
 	uint64 SizeInBytes{};
+	FRHITransientResourceStats TransientStats;
 
 	TArray<FRDGPassHandle> Passes;
 
@@ -152,6 +153,9 @@ public:
 	uint32 ScopeDepth{};
 	uint32 PassCount{};
 
+	FRHITransientHeapStats TransientHeapStats;
+	TArray<uint64, TInlineAllocator<8>> TransientHeapOffsets;
+
 	FGraphPacket(TraceServices::ILinearAllocator& Allocator, const UE::Trace::IAnalyzer::FOnEventContext& Context);
 
 	const FPassPacket* GetProloguePass() const
@@ -211,5 +215,15 @@ private:
 	TGraphTimeline GraphTimeline;
 };
 
+inline bool Intersects(const FPassIntervalPacket& A, const FPassPacket& B)
+{
+	return A.FirstPass <= B.Handle && A.LastPass >= B.Handle;
+}
+
+inline bool Intersects(const FPassIntervalPacket& A, const FPassIntervalPacket& B)
+{
+	return !(B.LastPass < A.FirstPass || A.LastPass < B.FirstPass);
+}
+
 } //namespace RenderGraphInsights
 } //namespace UE
diff --git a/Engine/Plugins/RenderGraphInsights/Source/RenderGraphInsights/Private/RenderGraphTimingViewSession.h b/Engine/Plugins/RenderGraphInsights/Source/RenderGraphInsights/Private/RenderGraphTimingViewSession.h
--- a/Engine/Plugins/RenderGraphInsights/Source/RenderGraphInsights/Private/RenderGraphTimingViewSession.h
+++ b/Engine/Plugins/RenderGraphInsights/Source/RenderGraphInsights/Private/RenderGraphTimingViewSession.h
@@ -36,6 +36,8 @@ public:
 	/** Show/Hide the RenderGraph track */
 	void ToggleRenderGraphTrack();
 
+	Insights::ITimingViewSession* GetTimingViewSession() const { return TimingViewSession; }
+
 private:
 	const TraceServices::IAnalysisSession* AnalysisSession = nullptr;
 	Insights::ITimingViewSession* TimingViewSession;
diff --git a/Engine/Plugins/RenderGraphInsights/Source/RenderGraphInsights/Private/RenderGraphTrack.cpp b/Engine/Plugins/RenderGraphInsights/Source/RenderGraphInsights/Private/RenderGraphTrack.cpp
--- a/Engine/Plugins/RenderGraphInsights/Source/RenderGraphInsights/Private/RenderGraphTrack.cpp
+++ b/Engine/Plugins/RenderGraphInsights/Source/RenderGraphInsights/Private/RenderGraphTrack.cpp
@@ -18,37 +18,123 @@
 #include "Insights/ViewModels/TimingEvent.h"
 #include "Insights/ViewModels/TimingEventSearch.h"
 #include "Insights/ViewModels/TooltipDrawState.h"
+#include "Insights/ITimingViewSession.h"
 
 #include "Widgets/Input/SSearchBox.h"
 #include "Widgets/Input/SSpinBox.h"
 
+#include "Brushes/SlateRoundedBoxBrush.h"
+
 #include "TraceServices/Model/Frames.h"
 
 #define LOCTEXT_NAMESPACE "RenderGraphTrack"
 
-namespace UE
+namespace UE { namespace RenderGraphInsights {
+
+constexpr uint32 kBuilderColor = 0xffa0a0a0;
+constexpr uint32 kRasterPassColor = 0xff7F2D2D;
+constexpr uint32 kComputePassColor = 0xff2D9F9F;
+constexpr uint32 kNoParameterPassColor = 0xff4D4D4D;
+constexpr uint32 kAsyncComputePassColor = 0xff2D7f2d;
+constexpr uint32 kTextureColor = 0xff89cff0;
+constexpr uint32 kBufferColor = 0xff66D066;
+
+static uint32 GetPassColor(const FPassPacket& Packet)
+{
+	const ERDGPassFlags Flags = Packet.Flags;
+
+	uint32 Color = 0;
+
+	const bool bNoParameterPass = (!Packet.Buffers.Num() && !Packet.Textures.Num());
+
+	if (bNoParameterPass && !Packet.bCulled)
+	{
+		Color = kNoParameterPassColor;
+	}
+	else if (EnumHasAnyFlags(Flags, ERDGPassFlags::AsyncCompute))
+	{
+		Color = kAsyncComputePassColor;
+	}
+	else if (EnumHasAnyFlags(Flags, ERDGPassFlags::Raster))
+	{
+		Color = kRasterPassColor;
+	}
+	else
+	{
+		Color = kComputePassColor;
+	}
+
+	if (Packet.bCulled)
+	{
+		Color &= 0x00FFFFFF;
+		Color |= 0x40000000;
+	}
+
+	return Color;
+}
+
+inline uint32 GetResourceColorBySize(uint64 Size, uint64 MaxSize)
 {
-namespace RenderGraphInsights
+	const FLinearColor Low(0.01, 0.01, 0.01, 0.25f);
+	const FLinearColor High(1.0, 0.1, 0.1, 1.0f);
+	const float Percentage = FMath::Sqrt(float(Size) / float(MaxSize));
+	return FLinearColor::LerpUsingHSV(Low, High, Percentage).ToFColor(false).ToPackedARGB();
+}
+
+///////////////////////////////////////////////////////////////////////////////
+
+template <typename VisibleItemType>
+static VisibleItemType& AddVisibleItem(TArray<VisibleItemType>& VisibleItems, const VisibleItemType& InVisibleItem)
+{
+	const uint32 VisibleIndex = VisibleItems.Num();
+
+	VisibleItemType& VisibleItem = VisibleItems.Emplace_GetRef(InVisibleItem);
+	check(VisibleItem.Packet->VisibleIndex == kInvalidVisibleIndex);
+	VisibleItem.Packet->VisibleIndex = VisibleIndex;
+	VisibleItem.Index = VisibleIndex;
+	return VisibleItem;
+}
+
+template <typename VisibleItemType>
+static void ResetVisibleItemArray(TArray<VisibleItemType>& VisibleItems)
 {
+	for (VisibleItemType& VisibleItem : VisibleItems)
+	{
+		check(VisibleItem.Packet->VisibleIndex == VisibleItem.Index);
+		VisibleItem.Packet->VisibleIndex = kInvalidVisibleIndex;
+	}
+	VisibleItems.Empty();
+}
+
+///////////////////////////////////////////////////////////////////////////////
 
-class FEvent : public FTimingEvent
+class FRenderGraphTimingEvent : public FTimingEvent
 {
-	INSIGHTS_DECLARE_RTTI(FEvent, FTimingEvent)
+	INSIGHTS_DECLARE_RTTI(FRenderGraphTimingEvent, FTimingEvent)
 public:
-	FEvent(const TSharedRef<const FBaseTimingTrack> InTrack, const FVisibleItem& InItem)
-		: FTimingEvent(InTrack, InItem.Packet->StartTime, InItem.Packet->EndTime, InItem.Depth)
+	FRenderGraphTimingEvent(const TSharedRef<const FBaseTimingTrack> InTrack, const FVisibleItem& InItem)
+		: FTimingEvent(InTrack, InItem.Packet->StartTime, InItem.Packet->EndTime, (int32)InItem.DepthH)
 	{}
 
+	virtual bool Equals(const ITimingEvent& Other) const override
+	{
+		if (Other.Is<FRenderGraphTimingEvent>())
+		{
+			return &GetPacket() == &Other.As<FRenderGraphTimingEvent>().GetPacket();
+		}
+		return false;
+	}
+
 	virtual const FVisibleItem& GetItem() const = 0;
 	virtual const FPacket& GetPacket() const = 0;
 };
 
 template <typename ItemType>
-class TEventHelper : public FEvent
+class TRenderGraphTimingEventHelper : public FRenderGraphTimingEvent
 {
 public:
-	TEventHelper(const TSharedRef<const FBaseTimingTrack>& InTrack, const ItemType& InItem)
-		: FEvent(InTrack, InItem)
+	TRenderGraphTimingEventHelper(const TSharedRef<const FBaseTimingTrack>& InTrack, const ItemType& InItem)
+		: FRenderGraphTimingEvent(InTrack, InItem)
 		, Item(InItem)
 	{
 		Item.Index = kInvalidVisibleIndex;
@@ -70,55 +156,104 @@ private:
 	ItemType Item;
 };
 
-class FVisibleScopeEvent final : public TEventHelper<FVisibleScope>
+class FVisibleScopeEvent final : public TRenderGraphTimingEventHelper<FVisibleScope>
 {
-	INSIGHTS_DECLARE_RTTI(FVisibleScopeEvent, FEvent)
+	INSIGHTS_DECLARE_RTTI(FVisibleScopeEvent, FRenderGraphTimingEvent)
 public:
 	FVisibleScopeEvent(const TSharedRef<const FBaseTimingTrack>& InTrack, const FVisibleScope& InItem)
-		: TEventHelper<FVisibleScope>(InTrack, InItem)
+		: TRenderGraphTimingEventHelper<FVisibleScope>(InTrack, InItem)
 	{}
 };
 
-class FVisiblePassEvent final : public TEventHelper<FVisiblePass>
+class FVisiblePassEvent final : public TRenderGraphTimingEventHelper<FVisiblePass>
 {
-	INSIGHTS_DECLARE_RTTI(FVisiblePassEvent, FEvent)
+	INSIGHTS_DECLARE_RTTI(FVisiblePassEvent, FRenderGraphTimingEvent)
 public:
 	FVisiblePassEvent(const TSharedRef<const FBaseTimingTrack>& InTrack, const FVisiblePass& InItem)
-		: TEventHelper<FVisiblePass>(InTrack, InItem)
+		: TRenderGraphTimingEventHelper<FVisiblePass>(InTrack, InItem)
 	{}
 };
 
-class FVisibleTextureEvent final : public TEventHelper<FVisibleTexture>
+class FVisibleTextureEvent final : public TRenderGraphTimingEventHelper<FVisibleTexture>
 {
-	INSIGHTS_DECLARE_RTTI(FVisibleTextureEvent, FEvent)
+	INSIGHTS_DECLARE_RTTI(FVisibleTextureEvent, FRenderGraphTimingEvent)
 public:
 	FVisibleTextureEvent(const TSharedRef<const FBaseTimingTrack>& InTrack, const FVisibleTexture& InItem)
-		: TEventHelper<FVisibleTexture>(InTrack, InItem)
+		: TRenderGraphTimingEventHelper<FVisibleTexture>(InTrack, InItem)
 	{}
 };
 
-class FVisibleBufferEvent final : public TEventHelper<FVisibleBuffer>
+class FVisibleBufferEvent final : public TRenderGraphTimingEventHelper<FVisibleBuffer>
 {
-	INSIGHTS_DECLARE_RTTI(FVisibleBufferEvent, FEvent)
+	INSIGHTS_DECLARE_RTTI(FVisibleBufferEvent, FRenderGraphTimingEvent)
 public:
 	FVisibleBufferEvent(const TSharedRef<const FBaseTimingTrack>& InTrack, const FVisibleBuffer& InItem)
-		: TEventHelper<FVisibleBuffer>(InTrack, InItem)
+		: TRenderGraphTimingEventHelper<FVisibleBuffer>(InTrack, InItem)
 	{}
 };
 
+inline const FPassPacket* TryGetPassPacket(const TSharedPtr<const ITimingEvent>& TimingEventBase)
+{
+	if (TimingEventBase)
+	{
+		if (TimingEventBase->Is<FVisiblePassEvent>())
+		{
+			return &TimingEventBase->As<FVisiblePassEvent>().GetPacket();
+		}
+	}
+	return nullptr;
+}
+
+///////////////////////////////////////////////////////////////////////////////
+// Packet Filter
+///////////////////////////////////////////////////////////////////////////////
+
 class FPacketFilter final : public ITimingEventFilter
 {
 	INSIGHTS_DECLARE_RTTI(FPacketFilter, ITimingEventFilter)
 public:
-	FPacketFilter(const TSharedPtr<const FEvent>& InEvent)
+	FPacketFilter(const TSharedPtr<const FRenderGraphTimingEvent>& InEvent)
 		: Event(InEvent)
 		, Packet(&InEvent->GetPacket())
 		, Graph(Packet->Graph)
 	{}
 	~FPacketFilter() {}
 
-	bool FilterPacket(const FPacket& Packet) const;
-	bool FilterPacketExact(const FPacket& Packet) const;
+	bool FilterPacket(const FPacket& PacketToFilter) const
+	{
+		if (FilterPacketExact(PacketToFilter))
+		{
+			return true;
+		}
+
+		if (PacketToFilter.Graph != Graph)
+		{
+			return false;
+		}
+
+		if (Packet->Is<FScopePacket>())
+		{
+			const FScopePacket& Scope = Packet->As<FScopePacket>();
+
+			if (PacketToFilter.Is<FPassIntervalPacket>())
+			{
+				const FPassIntervalPacket& PassIntervalToFilter = PacketToFilter.As<FPassIntervalPacket>();
+				return Intersects(Scope, PassIntervalToFilter);
+			}
+			else if (PacketToFilter.Is<FPassPacket>())
+			{
+				const FPassPacket& PassToFilter = PacketToFilter.As<FPassPacket>();
+				return Intersects(Scope, PassToFilter);
+			}
+		}
+
+		return false;
+	}
+
+	bool FilterPacketExact(const FPacket& PacketToFilter) const
+	{
+		return Packet == &PacketToFilter;
+	}
 
 	const FPacket& GetPacket() const
 	{
@@ -138,102 +273,50 @@ private:
 	uint32 GetChangeNumber() const override { return 0; }
 	//! ITimingEventFilter
 
-	TSharedPtr<const FEvent> Event;
+	TSharedPtr<const FRenderGraphTimingEvent> Event;
 	const FPacket* Packet{};
 	const FGraphPacket* Graph{};
 };
 
-INSIGHTS_IMPLEMENT_RTTI(FRenderGraphTrack)
-INSIGHTS_IMPLEMENT_RTTI(FPacketFilter)
-INSIGHTS_IMPLEMENT_RTTI(FEvent)
-INSIGHTS_IMPLEMENT_RTTI(FVisibleScopeEvent)
-INSIGHTS_IMPLEMENT_RTTI(FVisiblePassEvent)
-INSIGHTS_IMPLEMENT_RTTI(FVisibleTextureEvent)
-INSIGHTS_IMPLEMENT_RTTI(FVisibleBufferEvent)
-
-template <typename VisibleItemType>
-static VisibleItemType& AddVisibleItem(TArray<VisibleItemType>& VisibleItems, const VisibleItemType& InVisibleItem)
+inline bool TryFilterPacket(const FPacketFilter* Filter, const FPacket& Packet)
 {
-	const uint32 VisibleIndex = VisibleItems.Num();
-
-	VisibleItemType& VisibleItem = VisibleItems.Emplace_GetRef(InVisibleItem);
-	check(VisibleItem.Packet->VisibleIndex == kInvalidVisibleIndex);
-	VisibleItem.Packet->VisibleIndex = VisibleIndex;
-	VisibleItem.Index = VisibleIndex;
-	return VisibleItem;
-}
-
-template <typename VisibleItemType>
-static void ResetVisibleItemArray(TArray<VisibleItemType>& VisibleItems)
-{
-	for (VisibleItemType& VisibleItem : VisibleItems)
-	{
-		check(VisibleItem.Packet->VisibleIndex == VisibleItem.Index);
-		VisibleItem.Packet->VisibleIndex = kInvalidVisibleIndex;
-	}
-	VisibleItems.Empty();
+	return Filter ? Filter->FilterPacket(Packet) : false;
 }
 
-static bool Intersects(const FPassIntervalPacket& A, const FPassPacket& B)
+inline bool TryFilterPacketExact(const FPacketFilter* Filter, const FPacket& Packet)
 {
-	return A.FirstPass <= B.Handle && A.LastPass >= B.Handle;
+	return Filter ? Filter->FilterPacketExact(Packet) : false;
 }
 
-static bool Intersects(const FPassIntervalPacket& A, const FPassIntervalPacket& B)
+inline const FPacketFilter* GetPacketFilter(const TSharedPtr<ITimingEventFilter>& EventFilter)
 {
-	return !(B.LastPass < A.FirstPass || A.LastPass < B.FirstPass);
+	if (EventFilter && EventFilter->Is<FPacketFilter>())
+	{
+		return &EventFilter->As<FPacketFilter>();
+	}
+	return nullptr;
 }
 
-static void AddEvent(ITimingEventsTrackDrawStateBuilder& Builder, const FVisibleItem& Item)
-{
-	Builder.AddEvent(Item.StartTime, Item.EndTime, Item.Depth, Item.Color, [&Item](float) { return Item.Name; });
-};
+///////////////////////////////////////////////////////////////////////////////
 
-bool FPacketFilter::FilterPacketExact(const FPacket& PacketToFilter) const
+FBoundingBox FVisibleItem::GetBoundingBox(const FTimingViewLayout& Layout) const
 {
-	return Packet == &PacketToFilter;
+	return FBoundingBox(FVector2D(SlateX, DepthY), FVector2D(SlateX + SlateW, DepthY + DepthH));
 }
 
-bool FPacketFilter::FilterPacket(const FPacket& PacketToFilter) const
+void FVisibleItem::Draw(FRenderGraphTrackDrawStateBuilder& Builder) const
 {
-	if (FilterPacketExact(PacketToFilter))
-	{
-		return true;
-	}
-
-	if (PacketToFilter.Graph != Graph)
+	if (DepthH != 1.0f || FMath::Fractional(DepthY) != 0.0f)
 	{
-		return false;
+		Builder.AddEvent(StartTime, EndTime, DepthY, DepthH, Color, [this](float) { return Name; });
 	}
-
-	if (Packet->Is<FScopePacket>())
+	else
 	{
-		const FScopePacket& Scope = Packet->As<FScopePacket>();
-
-		if (PacketToFilter.Is<FPassIntervalPacket>())
-		{
-			const FPassIntervalPacket& PassIntervalToFilter = PacketToFilter.As<FPassIntervalPacket>();
-			return Intersects(Scope, PassIntervalToFilter);
-		}
-		else if (PacketToFilter.Is<FPassPacket>())
-		{
-			const FPassPacket& PassToFilter = PacketToFilter.As<FPassPacket>();
-			return Intersects(Scope, PassToFilter);
-		}
+		Builder.AddEvent(StartTime, EndTime, DepthY, Color, [this](float) { return Name; });
 	}
-
-	return false;
-}
-
-bool FRenderGraphTrack::FilterPacket(const FPacket& Packet) const
-{
-	return PacketFilter ? PacketFilter->FilterPacket(Packet) : false;
 }
 
-bool FRenderGraphTrack::FilterPacketExact(const FPacket& Packet) const
-{
-	return PacketFilter ? PacketFilter->FilterPacketExact(Packet) : false;
-}
+///////////////////////////////////////////////////////////////////////////////
 
 void FVisibleGraph::Reset()
 {
@@ -261,18 +344,22 @@ void FVisibleGraph::AddPass(const FVisiblePass& InVisiblePass)
 void FVisibleGraph::AddTexture(const FVisibleTexture& VisibleTexture)
 {
 	AddVisibleItem(Textures, VisibleTexture);
+
+	DepthH = FMath::Max(DepthH, VisibleTexture.DepthY + VisibleTexture.DepthH);
 }
 
 void FVisibleGraph::AddBuffer(const FVisibleBuffer& VisibleBuffer)
 {
 	AddVisibleItem(Buffers, VisibleBuffer);
+
+	DepthH = FMath::Max(DepthH, VisibleBuffer.DepthY + VisibleBuffer.DepthH);
 }
 
-const FVisibleItem* FVisibleGraph::FindItem(float InPosX, float InPosY) const
+const FVisibleItem* FVisibleGraph::FindItem(float InSlateX, float InDepthY) const
 {
 	for (const FVisibleTexture& Texture : Textures)
 	{
-		if (Texture.Intersects(InPosX, InPosY))
+		if (Texture.Intersects(InSlateX, InDepthY))
 		{
 			return &Texture;
 		}
@@ -280,7 +367,7 @@ const FVisibleItem* FVisibleGraph::FindItem(float InPosX, float InPosY) const
 
 	for (const FVisibleBuffer& Buffer : Buffers)
 	{
-		if (Buffer.Intersects(InPosX, InPosY))
+		if (Buffer.Intersects(InSlateX, InDepthY))
 		{
 			return &Buffer;
 		}
@@ -288,7 +375,7 @@ const FVisibleItem* FVisibleGraph::FindItem(float InPosX, float InPosY) const
 
 	for (const FVisibleScope& Scope : Scopes)
 	{
-		if (Scope.Intersects(InPosX, InPosY))
+		if (Scope.Intersects(InSlateX, InDepthY))
 		{
 			return &Scope;
 		}
@@ -296,7 +383,7 @@ const FVisibleItem* FVisibleGraph::FindItem(float InPosX, float InPosY) const
 
 	for (const FVisiblePass& Pass : Passes)
 	{
-		if (Pass.Intersects(InPosX, InPosY))
+		if (Pass.Intersects(InSlateX, InDepthY))
 		{
 			return &Pass;
 		}
@@ -305,62 +392,31 @@ const FVisibleItem* FVisibleGraph::FindItem(float InPosX, float InPosY) const
 	return nullptr;
 }
 
-FRenderGraphTrack::FRenderGraphTrack(const FRenderGraphTimingViewSession& InSharedData)
-	: Super(LOCTEXT("TrackNameFormat", "RDG").ToString())
-	, SharedData(InSharedData)
-{}
-
-static const uint32 kBuilderColor = 0xffa0a0a0;
-static const uint32 kRasterPassColor = 0xff7F2D2D;
-static const uint32 kComputePassColor = 0xff2D9F9F;
-static const uint32 kNoParameterPassColor = 0xff4D4D4D;
-static const uint32 kAsyncComputePassColor = 0xff2D7f2d;
-static const uint32 kTextureColor = 0xff89cff0;
-static const uint32 kBufferColor = 0xff66D066;
-
-static const float kMinGraphPixels = 5.0f;
-static const float kMinPassMarginPixels = 5.0f;
-
-static uint32 GetPassColor(const FPassPacket& Packet)
+void FVisibleGraph::Draw(FRenderGraphTrackDrawStateBuilder& Builder) const
 {
-	const ERDGPassFlags Flags = Packet.Flags;
+	Builder.AddEvent(StartTime, EndTime, 0, Name, 0, kBuilderColor);
+}
 
-	uint32 Color = 0;
+///////////////////////////////////////////////////////////////////////////////
 
-	const bool bNoParameterPass = (!Packet.Buffers.Num() && !Packet.Textures.Num());
-
-	if (bNoParameterPass && !Packet.bCulled)
-	{
-		Color = kNoParameterPassColor;
-	}
-	else if (EnumHasAnyFlags(Flags, ERDGPassFlags::AsyncCompute))
-	{
-		Color = kAsyncComputePassColor;
-	}
-	else if (EnumHasAnyFlags(Flags, ERDGPassFlags::Raster))
-	{
-		Color = kRasterPassColor;
-	}
-	else
-	{
-		Color = kComputePassColor;
-	}
-
-	if (Packet.bCulled)
-	{
-		Color &= 0x00FFFFFF;
-		Color |= 0x40000000;
-	}
+INSIGHTS_IMPLEMENT_RTTI(FRenderGraphTimingEvent)
+INSIGHTS_IMPLEMENT_RTTI(FPacketFilter)
+INSIGHTS_IMPLEMENT_RTTI(FVisibleScopeEvent)
+INSIGHTS_IMPLEMENT_RTTI(FVisiblePassEvent)
+INSIGHTS_IMPLEMENT_RTTI(FVisibleTextureEvent)
+INSIGHTS_IMPLEMENT_RTTI(FVisibleBufferEvent)
+INSIGHTS_IMPLEMENT_RTTI(FRenderGraphTrack)
 
-	return Color;
-}
+///////////////////////////////////////////////////////////////////////////////
 
-static uint32 GetColorBySize(uint64 Size, uint64 MaxSize)
+FRenderGraphTrack::FRenderGraphTrack(const FRenderGraphTimingViewSession& InSharedData)
+	: Super(LOCTEXT("TrackNameFormat", "RDG").ToString())
+	, SharedData(InSharedData)
+	, NumLanes(0)
+	, DrawState(MakeShared<FRenderGraphTrackDrawState>())
+	, FilteredDrawState(MakeShared<FRenderGraphTrackDrawState>())
 {
-	const FLinearColor Low(0.01, 0.01, 0.01, 0.25f);
-	const FLinearColor High(1.0, 0.1, 0.1, 1.0f);
-	const float Percentage = FMath::Sqrt(float(Size) / float(MaxSize));
-	return FLinearColor::LerpUsingHSV(Low, High, Percentage).ToFColor(false).ToPackedARGB();
+	SetValidLocations(ETimingTrackLocation::Scrollable | ETimingTrackLocation::TopDocked | ETimingTrackLocation::BottomDocked);
 }
 
 uint32 FRenderGraphTrack::GetTextureColor(const FTexturePacket& Texture, uint64 MaxSizeInBytes) const
@@ -370,7 +426,7 @@ uint32 FRenderGraphTrack::GetTextureColor(const FTexturePacket& Texture, uint64
 		return kTextureColor;
 	}
 
-	return GetColorBySize(Texture.SizeInBytes, MaxSizeInBytes);
+	return GetResourceColorBySize(Texture.SizeInBytes, MaxSizeInBytes);
 }
 
 uint32 FRenderGraphTrack::GetBufferColor(const FBufferPacket& Texture, uint64 MaxSizeInBytes) const
@@ -380,29 +436,167 @@ uint32 FRenderGraphTrack::GetBufferColor(const FBufferPacket& Texture, uint64 Ma
 		return kBufferColor;
 	}
 
-	return GetColorBySize(Texture.SizeInBytes, MaxSizeInBytes);
+	return GetResourceColorBySize(Texture.SizeInBytes, MaxSizeInBytes);
 }
 
-void FRenderGraphTrack::Update(const ITimingTrackUpdateContext& Context)
+float FRenderGraphTrack::TransientHeapOffsetToDepth(uint64 HeapOffset) const
 {
-	FTimingEventsTrack::Update(Context);
+	return (float)HeapOffset / TransientHeapBytesPerDepthSlot;
+}
 
-	MousePosition = Context.GetMousePosition();
-	MousePosition.Y -= GetPosY();
+float FRenderGraphTrack::TransientHeapOffsetToSlateY(const FTimingViewLayout& Layout, uint64 HeapOffset) const
+{
+	return (Layout.EventH + Layout.EventDY) * TransientHeapOffsetToDepth(HeapOffset);
+}
 
-	SelectedPass = nullptr;
+void FRenderGraphTrack::Reset()
+{
+	FBaseTimingTrack::Reset();
 
-	const TSharedPtr<const ITimingEvent> SelectedEventPtr = Context.GetSelectedEvent();
-	if (SelectedEventPtr.IsValid() && SelectedEventPtr->Is<FEvent>())
+	NumLanes = 0;
+	DrawState->Reset();
+	FilteredDrawState->Reset();
+}
+
+void FRenderGraphTrack::PreUpdate(const ITimingTrackUpdateContext& Context)
+{
+	const FTimingTrackViewport& Viewport = Context.GetViewport();
+	const FTimingViewLayout& Layout = Viewport.GetLayout();
+
+	if (IsDirty() || Viewport.IsHorizontalViewportDirty())
+	{
+		ClearDirtyFlag();
+
+		for (FVisibleGraph& VisibleGraph : VisibleGraphs)
+		{
+			VisibleGraph.Reset();
+		}
+		ResetVisibleItemArray(VisibleGraphs);
+
+		int32 MaxDepth = -1;
+
+		{
+			FRenderGraphTrackDrawStateBuilder Builder(*DrawState, Viewport);
+
+			BuildDrawState(Builder, Context);
+
+			Builder.Flush();
+
+			if (Builder.GetMaxDepth() > MaxDepth)
+			{
+				MaxDepth = Builder.GetMaxDepth();
+			}
+		}
+
+		const TSharedPtr<ITimingEventFilter> EventFilter = Context.GetEventFilter();
+		if ((EventFilter.IsValid() && EventFilter->FilterTrack(*this)))
+		{
+			const bool bFastLastBuild = FilteredDrawStateInfo.LastBuildDuration < 0.005; // LastBuildDuration < 5ms
+			const bool bFilterPointerChanged = !FilteredDrawStateInfo.LastEventFilter.HasSameObject(EventFilter.Get());
+
+			if (bFastLastBuild || bFilterPointerChanged)
+			{
+				FilteredDrawStateInfo.LastEventFilter = EventFilter;
+				FilteredDrawStateInfo.ViewportStartTime = Viewport.GetStartTime();
+				FilteredDrawStateInfo.ViewportScaleX = Viewport.GetScaleX();
+				FilteredDrawStateInfo.Counter = 0;
+			}
+			else
+			{
+				if (FilteredDrawStateInfo.ViewportStartTime == Viewport.GetStartTime() &&
+					FilteredDrawStateInfo.ViewportScaleX == Viewport.GetScaleX())
+				{
+					if (FilteredDrawStateInfo.Counter > 0)
+					{
+						FilteredDrawStateInfo.Counter--;
+					}
+				}
+				else
+				{
+					FilteredDrawStateInfo.ViewportStartTime = Viewport.GetStartTime();
+					FilteredDrawStateInfo.ViewportScaleX = Viewport.GetScaleX();
+					FilteredDrawStateInfo.Counter = 1; // wait
+				}
+			}
+
+			if (FilteredDrawStateInfo.Counter == 0)
+			{
+				const uint64 StartTime = FPlatformTime::Cycles64();
+				{
+					FRenderGraphTrackDrawStateBuilder Builder(*FilteredDrawState, Viewport);
+					BuildFilteredDrawState(Builder, Context);
+					Builder.Flush();
+				}
+				const uint64 EndTime = FPlatformTime::Cycles64();
+				FilteredDrawStateInfo.LastBuildDuration = static_cast<double>(EndTime - StartTime) * FPlatformTime::GetSecondsPerCycle64();
+			}
+			else
+			{
+				FilteredDrawState->Reset();
+				FilteredDrawStateInfo.Opacity = 0.0f;
+				SetDirtyFlag();
+			}
+		}
+		else
+		{
+			FilteredDrawStateInfo.LastBuildDuration = 0.0;
+
+			if (FilteredDrawStateInfo.LastEventFilter.IsValid())
+			{
+				FilteredDrawStateInfo.LastEventFilter.Reset();
+				FilteredDrawStateInfo.Counter = 0;
+				FilteredDrawState->Reset();
+			}
+		}
+
+		NumLanes = MaxDepth + 1;
+	}
+
+	const float CurrentTrackHeight = GetHeight();
+	const float DesiredTrackHeight = Layout.ComputeTrackHeight(NumLanes);
+
+	if (CurrentTrackHeight < DesiredTrackHeight)
 	{
-		if (SelectedEventPtr->Is<FVisiblePassEvent>())
+		float NewTrackHeight;
+		if (Viewport.IsDirty(ETimingTrackViewportDirtyFlags::VLayoutChanged))
+		{
+			NewTrackHeight = DesiredTrackHeight;
+		}
+		else
 		{
-			SelectedPass = &SelectedEventPtr->As<FVisiblePassEvent>().GetPacket();
+			NewTrackHeight = FMath::CeilToFloat(CurrentTrackHeight * 0.9f + DesiredTrackHeight * 0.1f);
 		}
+		SetHeight(NewTrackHeight);
+	}
+	else if (CurrentTrackHeight > DesiredTrackHeight)
+	{
+		float NewTrackHeight;
+		if (Viewport.IsDirty(ETimingTrackViewportDirtyFlags::VLayoutChanged))
+		{
+			NewTrackHeight = DesiredTrackHeight;
+		}
+		else
+		{
+			NewTrackHeight = FMath::FloorToFloat(CurrentTrackHeight * 0.9f + DesiredTrackHeight * 0.1f);
+		}
+		SetHeight(NewTrackHeight);
+	}
+}
+
+void FRenderGraphTrack::Update(const ITimingTrackUpdateContext& Context)
+{
+	const FTimingTrackViewport& Viewport = Context.GetViewport();
+	const FTimingViewLayout& Layout = Viewport.GetLayout();
 
-		const FEvent& SelectedEvent = SelectedEventPtr->As<FEvent>();
+	const FVector2D MousePosition = Context.GetMousePosition();
+	MouseSlateX = MousePosition.X;
+	MouseSlateY = MousePosition.Y - GetPosY();
+	MouseDepthY = GetDepthFromLaneY(Layout, MouseSlateY);
 
-		const FTimingTrackViewport& Viewport = Context.GetViewport();
+	const TSharedPtr<const ITimingEvent> SelectedEventPtr = Context.GetSelectedEvent();
+	if (SelectedEventPtr.IsValid() && SelectedEventPtr->Is<FRenderGraphTimingEvent>())
+	{
+		const FRenderGraphTimingEvent& SelectedEvent = SelectedEventPtr->As<FRenderGraphTimingEvent>();
 
 		const FVector2D ViewportPos(Viewport.GetWidth(), Viewport.GetHeight());
 
@@ -417,182 +611,165 @@ void FRenderGraphTrack::Update(const ITimingTrackUpdateContext& Context)
 	}
 }
 
+void FRenderGraphTrack::PostUpdate(const ITimingTrackUpdateContext& Context)
+{
+	constexpr float HeaderWidth = 100.0f;
+	constexpr float HeaderHeight = 14.0f;
+
+	const float MouseY = Context.GetMousePosition().Y;
+	if (MouseY >= GetPosY() && MouseY < GetPosY() + GetHeight())
+	{
+		const float MouseX = Context.GetMousePosition().X;
+		SetHeaderHoveredState(MouseX < HeaderWidth&& MouseY < GetPosY() + HeaderHeight);
+		SetHoveredState(true);
+	}
+	else
+	{
+		SetHoveredState(false);
+	}
+}
+
 void FRenderGraphTrack::Draw(const ITimingTrackDrawContext& Context) const
 {
-	FTimingEventsTrack::Draw(Context);
+	FRenderGraphTrackDrawHelper DrawHelper(Context);
+
+	const bool bFilterActive = Context.GetEventFilter().IsValid() && Context.GetEventFilter()->FilterTrack(*this);
 
-	const FEvent* HoveredEvent = nullptr;
-	if (TSharedPtr<const ITimingEvent> HoveredEventBase = Context.GetHoveredEvent())
+	if (bFilterActive)
 	{
-		if (HoveredEventBase->Is<FEvent>())
+		DrawHelper.DrawFadedEvents(*DrawState, *this, 0.1f);
+
+		if (FilteredDrawStateInfo.Opacity == 1.0f)
 		{
-			HoveredEvent = &HoveredEventBase->As<FEvent>();
+			DrawHelper.DrawEvents(*FilteredDrawState, *this);
 		}
+		else
+		{
+			FilteredDrawStateInfo.Opacity = FMath::Min(1.0f, FilteredDrawStateInfo.Opacity + 0.05f);
+
+			DrawHelper.DrawFadedEvents(*FilteredDrawState, *this, FilteredDrawStateInfo.Opacity);
+		}
+	}
+	else
+	{
+		DrawHelper.DrawEvents(*DrawState, *this);
 	}
 
-	const int32 LineLayerId = Context.GetHelper().GetFirstLayerId() - 1;
+	DrawHelper.DrawTrackHeader(*this);
+
 	const FTimingTrackViewport& Viewport = Context.GetViewport();
+	const FTimingViewLayout& Layout = Viewport.GetLayout();
 	const FDrawContext& DrawContext = Context.GetDrawContext();
-	const ITimingViewDrawHelper& DrawHelper = Context.GetHelper();
-	const FLinearColor EdgeColor = DrawHelper.GetEdgeColor();
 
-	const float TrackY = GetPosY();
-	const float TrackH = GetHeight();
-	const float EventH = Viewport.GetLayout().EventH;
+	const FLinearColor EdgeColor = DrawHelper.GetEdgeColor();
 	const float ViewportWidth = Viewport.GetWidth();
 
-	const auto DrawClampedBox = [&](int32 LayerId, float MinX, float MinY, float W, float H, const FLinearColor& Color)
-	{
-		const float Guardband = 1024.0f;
-		float MaxX = MinX + W;
+	const FPassPacket* SelectedPass = TryGetPassPacket(Context.GetSelectedEvent());
+	const FPassPacket* HoveredPass = TryGetPassPacket(Context.GetHoveredEvent());
+	const FPacketFilter* PacketFilter = GetPacketFilter(Context.GetEventFilter());
 
-		if (MinX > ViewportWidth || MaxX < 0.0f)
+	for (const FVisibleGraph& VisibleGraph : VisibleGraphs)
+	{
+		if (VisibleGraph.IsCulled())
 		{
-			return;
+			continue;
 		}
 
-		MinX = FMath::Max(MinX, -Guardband);
-		MaxX = FMath::Min(MaxX, ViewportWidth + Guardband);
-		W = MaxX - MinX;
+		const FGraphPacket& Graph = VisibleGraph.GetPacket();
+		const float GraphWidth = VisibleGraph.SlateW;
 
-		DrawContext.DrawBox(LayerId, MinX, MinY, W, H, DrawHelper.GetWhiteBrush(), Color);
-	};
+		////////////////////////////////////////////////////////////////////////////////////
+		// Draw pass columns and markers.
+		////////////////////////////////////////////////////////////////////////////////////
 
-	const auto DrawClampedSpline = [&](int32 SplineLayerId, FSpline Spline)
-	{
-		const float Guardband = 1024.0f;
-		float MinX = Spline.Start.X;
-		float MaxX = Spline.Start.X + Spline.End.X;
+		const float PassLineStrideMin = 5.0f;
+		const float PassLineStride = GraphWidth / float(Graph.PassCount);
 
-		if (MinX > ViewportWidth || MaxX < 0.0f)
+		if (GraphWidth >= PassLineStrideMin)
 		{
-			return;
+			DrawHelper.DrawBox(*this, VisibleGraph.SlateX, 0.0f, 1.0f, VisibleGraph.DepthH, EdgeColor, EDrawLayer::Background);
+			DrawHelper.DrawBox(*this, VisibleGraph.SlateX + VisibleGraph.SlateW, 0.0f, 1.0f, VisibleGraph.DepthH, EdgeColor, EDrawLayer::Background);
 		}
 
-		MinX = FMath::Max(MinX, -Guardband);
-		MaxX = FMath::Min(MaxX, ViewportWidth + Guardband);
+		float RenderPassMergeMinX = 0.0f;
+		float ParallelExecuteMinX = 0.0f;
 
-		const float Width = MaxX - MinX;
+		constexpr uint32 BarCount = 2; // Render Pass Merge + Parallel Execute
+		constexpr float BarDepth  = 0.3f;
+		constexpr float BarMargin = (1.0f / float(BarCount + 1)) - BarDepth;
 
-		const float MinY = TrackY + Spline.Start.Y;
-
-		Spline.Start.X = MinX;
-		Spline.End.X = Width;
-		Spline.StartDir.X = FMath::Min(Guardband, Spline.StartDir.X);
-		Spline.EndDir.X = FMath::Min(Guardband, Spline.EndDir.X);
-
-		DrawContext.DrawSpline(SplineLayerId, MinX, MinY, FVector2D::ZeroVector, Spline.StartDir, Spline.End, Spline.EndDir, Spline.Thickness, Spline.Tint);
-	};
-
-	if (TrackH > 0.0f)
-	{
-		const FTimingViewLayout& Layout = Viewport.GetLayout();
-
-		for (const FVisibleGraph& VisibleGraph : VisibleGraphs)
+		for (uint32 PassIndex = 0; PassIndex < Graph.PassCount; ++PassIndex)
 		{
-			const FGraphPacket& Graph = VisibleGraph.GetPacket();
-			const uint32 PassCount = Graph.PassCount;
-			const float LineStrideMin = 5.0f;
-			const float GraphLineStride = VisibleGraph.Max.X - VisibleGraph.Min.X;
+			const FPassPacket& Pass = Graph.Passes[PassIndex];
+			const FVisiblePass& VisiblePass = VisibleGraph.GetVisiblePass(Pass);
+			const FBoundingBox VisiblePassBox = VisiblePass.GetBoundingBox(Layout);
 
-			if (GraphLineStride >= LineStrideMin)
+			if (Pass.bParallelExecuteBegin)
 			{
-				DrawClampedBox(LineLayerId, VisibleGraph.Min.X, TrackY, 1.0f, TrackH, EdgeColor);
-				DrawClampedBox(LineLayerId, VisibleGraph.Max.X, TrackY, 1.0f, TrackH, EdgeColor);
+				ParallelExecuteMinX = VisiblePassBox.Min.X;
 			}
 
-			const float PassLineStride = GraphLineStride / float(PassCount);
-
-			FVector2D RenderPassMergeMin = FVector2D::ZeroVector;
-			FVector2D ParallelExecuteMin = FVector2D::ZeroVector;
-
-			for (uint32 PassIndex = 0; PassIndex < PassCount; ++PassIndex)
+			if (Pass.bParallelExecuteEnd)
 			{
-				const FPassPacket& Pass = Graph.Passes[PassIndex];
-				const FVisiblePass& VisiblePass = VisibleGraph.GetVisiblePass(Pass);
+				const float W = VisiblePassBox.Max.X - ParallelExecuteMinX;
+				const float H = BarDepth;
+				const float X = ParallelExecuteMinX;
+				const float Y = VisibleGraph.HeaderPassBarDepth + BarMargin;
 
-				if (Pass.bParallelExecuteBegin)
-				{
-					ParallelExecuteMin = VisiblePass.Min;
-				}
-
-				if (Pass.bParallelExecuteEnd)
-				{
-					const FVector2D ParallelExecuteMax = VisiblePass.Max;
-					const float RenderPassMarginY = 10.0f;
-					const float W = ParallelExecuteMax.X - ParallelExecuteMin.X;
-					const float H = (ParallelExecuteMax.Y - ParallelExecuteMin.Y) * 0.25f;
-					const float X = ParallelExecuteMin.X;
-					const float Y = ParallelExecuteMin.Y - H - RenderPassMarginY;
-
-					DrawClampedBox(LineLayerId, X, TrackY + Y, W, H, FLinearColor(1.0f, 1.0f, 1.0f, 0.75f));
-				}
+				DrawHelper.DrawBox(*this, X, Y, W, H, FLinearColor(1.0f, 1.0f, 1.0f, 0.75f), EDrawLayer::Background);
+			}
 
-				if (!Pass.bSkipRenderPassBegin && Pass.bSkipRenderPassEnd)
-				{
-					RenderPassMergeMin = VisiblePass.Min;
-				}
+			if (!Pass.bSkipRenderPassBegin && Pass.bSkipRenderPassEnd)
+			{
+				RenderPassMergeMinX = VisiblePassBox.Min.X;
+			}
 
-				if (Pass.bSkipRenderPassBegin && !Pass.bSkipRenderPassEnd)
-				{
-					const FVector2D RenderPassMergeMax = VisiblePass.Max;
-					const float RenderPassMarginY = 3.0f;
-					const float W = RenderPassMergeMax.X - RenderPassMergeMin.X;
-					const float H = (RenderPassMergeMax.Y - RenderPassMergeMin.Y) * 0.25f;
-					const float X = RenderPassMergeMin.X;
-					const float Y = RenderPassMergeMin.Y - H - RenderPassMarginY;
-
-					DrawClampedBox(LineLayerId, X, TrackY + Y, W, H, FLinearColor(0.8f, 0.2f, 0.2f, 0.75f));
-				}
+			if (Pass.bSkipRenderPassBegin && !Pass.bSkipRenderPassEnd)
+			{
+				const float W = VisiblePassBox.Max.X - RenderPassMergeMinX;
+				const float H = BarDepth;
+				const float X = RenderPassMergeMinX;
+				const float Y = VisibleGraph.HeaderPassBarDepth + BarMargin * 2.0f + BarDepth;
 
-				const float X = Viewport.TimeToSlateUnitsRounded(Pass.StartTime);
-				const float Y = TrackY + VisiblePass.Max.Y;
-				const float H = TrackH - VisiblePass.Max.Y;
+				DrawHelper.DrawBox(*this, X, Y, W, H, FLinearColor(0.8f, 0.2f, 0.2f, 0.75f), EDrawLayer::Background);
+			}
 
-				const bool bHoveredPass = HoveredEvent && &HoveredEvent->GetPacket() == &Pass;
-				const bool bSelectedPass = SelectedPass == &Pass;
-				const bool bFilteredPass = FilterPacket(Pass);
+			const float X = Viewport.TimeToSlateUnitsRounded(Pass.StartTime);
+			const float Y = VisiblePassBox.Max.Y;
+			const float H = VisibleGraph.DepthH - VisiblePassBox.Max.Y;
 
-				if (bHoveredPass || bFilteredPass || bSelectedPass)
-				{
-					const float W = PassLineStride + 1;
+			const bool bHoveredPass = HoveredPass == &Pass;
+			const bool bSelectedPass = SelectedPass == &Pass;
+			const bool bFilteredPass = TryFilterPacket(PacketFilter, Pass);
 
-					const FLinearColor EdgeColorTranslucent = FLinearColor(EdgeColor.R, EdgeColor.G, EdgeColor.B, bFilteredPass ? 1.0f : 0.5f);
+			if (bHoveredPass || bFilteredPass || bSelectedPass)
+			{
+				const float W = PassLineStride + 1;
+				const FLinearColor EdgeColorTranslucent = FLinearColor(EdgeColor.R, EdgeColor.G, EdgeColor.B, bFilteredPass ? 1.0f : 0.5f);
 
-					DrawClampedBox(LineLayerId, X, Y, W, H, EdgeColorTranslucent);
-				}
-				else if (PassLineStride >= LineStrideMin && PassIndex != 0)
-				{
-					const float W = 1.0f;
+				DrawHelper.DrawBox(*this, X, Y, W, H, EdgeColorTranslucent, EDrawLayer::Background);
+			}
+			else if (PassLineStride >= PassLineStrideMin && PassIndex != 0)
+			{
+				const float W = 1.0f;
 
-					DrawClampedBox(LineLayerId, X, Y, W, H, EdgeColor);
-				}
+				DrawHelper.DrawBox(*this, X, Y, W, H, EdgeColor, EDrawLayer::Background);
 			}
 		}
-	}
-
-	const int32 SplineOverLayerId = LineLayerId + Context.GetHelper().GetNumLayerIds();
-	const int32 SplineUnderLineLayerId = LineLayerId + 1;
-
-	for (const FSpline& Spline : Splines)
-	{
-		DrawClampedSpline(SplineOverLayerId, Spline);
-	}
 
-	const bool bFilterActive = Context.GetEventFilter() != nullptr;
-
-	for (const FVisibleGraph& VisibleGraph : VisibleGraphs)
-	{
-		const FGraphPacket& Graph = VisibleGraph.GetPacket();
+		////////////////////////////////////////////////////////////////////////////////////
+		// Draw async compute fork / join fences
+		////////////////////////////////////////////////////////////////////////////////////
 
 		for (uint32 VisibleIndex : VisibleGraph.AsyncComputePasses)
 		{
 			const FVisiblePass& AsyncComputeVisiblePass = VisibleGraph.Passes[VisibleIndex];
 			const FPassPacket& AsyncComputePass = AsyncComputeVisiblePass.GetPacket();
+			const FBoundingBox AsyncComputeVisiblePassBox = AsyncComputeVisiblePass.GetBoundingBox(Layout);
 
 			float TintAlpha = 0.25f;
 
-			if (!bFilterActive || FilterPacket(AsyncComputePass))
+			if (!bFilterActive || TryFilterPacket(PacketFilter, AsyncComputePass))
 			{
 				TintAlpha = 0.75f;
 			}
@@ -603,48 +780,132 @@ void FRenderGraphTrack::Draw(const ITimingTrackDrawContext& Context) const
 
 			const FRDGPassHandle GraphicsForkHandle = AsyncComputePass.GraphicsForkPass;
 
-			if (GraphicsForkHandle.IsValid() && AsyncComputePass.bAsyncComputeBegin)
-			{
-				const FPassPacket& GraphicsForkPass = *Graph.GetPass(AsyncComputePass.GraphicsForkPass);
-				const FVisiblePass& GraphicsForkVisiblePass = VisibleGraph.GetVisiblePass(GraphicsForkPass);
+			if (GraphicsForkHandle.IsValid() && AsyncComputePass.bAsyncComputeBegin)
+			{
+				const FPassPacket& GraphicsForkPass = *Graph.GetPass(AsyncComputePass.GraphicsForkPass);
+				const FVisiblePass& GraphicsForkVisiblePass = VisibleGraph.GetVisiblePass(GraphicsForkPass);
+
+				const float X = GraphicsForkVisiblePass.SlateX + GraphicsForkVisiblePass.SlateW * EndT;
+				const float Y = GraphicsForkVisiblePass.DepthY + GraphicsForkVisiblePass.DepthH;
+				const float LocalEndX = AsyncComputeVisiblePass.SlateX + (AsyncComputeVisiblePass.SlateW * StartT) - X;
+				const float LocalEndY = AsyncComputeVisiblePass.DepthY - Y;
+
+				FSplinePrimitive Spline;
+				Spline.Start.X = X;
+				Spline.Start.Y = Y;
+				Spline.StartDir = FVector2D(0, SplineDir);
+				Spline.End = FVector2D(LocalEndX, LocalEndY);
+				Spline.EndDir = FVector2D(0, SplineDir);
+				Spline.Thickness = 2.0f;
+				Spline.Tint = FLinearColor(0.4f, 1.0f, 0.4f, TintAlpha);
+				DrawHelper.DrawSpline(*this, Spline, EDrawLayer::EventFill);
+			}
+
+			const FRDGPassHandle GraphicsJoinHandle = AsyncComputePass.GraphicsJoinPass;
+
+			if (GraphicsJoinHandle.IsValid() && AsyncComputePass.bAsyncComputeEnd)
+			{
+				const FPassPacket& GraphicsJoinPass = *Graph.GetPass(AsyncComputePass.GraphicsJoinPass);
+				const FVisiblePass& GraphicsJoinVisiblePass = VisibleGraph.GetVisiblePass(GraphicsJoinPass);
+
+				const float X = AsyncComputeVisiblePass.SlateX + AsyncComputeVisiblePass.SlateW * EndT;
+				const float Y = AsyncComputeVisiblePass.DepthY;
+				const float LocalEndX = GraphicsJoinVisiblePass.SlateX + (GraphicsJoinVisiblePass.SlateW * StartT) - X;
+				const float LocalEndY = GraphicsJoinVisiblePass.DepthY + GraphicsJoinVisiblePass.DepthH - Y;
+
+				FSplinePrimitive Spline;
+				Spline.Start.X = X;
+				Spline.Start.Y = Y;
+				Spline.StartDir = FVector2D(0, -SplineDir);
+				Spline.End = FVector2D(LocalEndX, LocalEndY);
+				Spline.EndDir = FVector2D(0, -SplineDir);
+				Spline.Thickness = 2.0f;
+				Spline.Tint = FLinearColor(1.0f, 0.4f, 0.4f, TintAlpha);
+				DrawHelper.DrawSpline(*this, Spline, EDrawLayer::EventFill);
+			}
+		}
+
+		////////////////////////////////////////////////////////////////////////////////////
+		// Draw transient heap vertical lines and size markers
+		////////////////////////////////////////////////////////////////////////////////////
+
+		if (Visualizer == EVisualizer::TransientHeaps)
+		{
+			const FLinearColor MBLineColor(0.8f, 0.1f, 0.1f, 0.2f);
+			const FLinearColor MBTextColor(0.8f, 0.1f, 0.1f, 0.8f);
+
+			const float StartDepth = VisibleGraph.HeaderEndDepth;
+
+			const bool bDrawText = GraphWidth > 100.0f;
+
+			const uint64 MiB = 1024 * 1024;
+
+			const auto DrawMibText = [&](float Depth, uint64 Offset, uint32 HeapIndex, uint32 HeapSizeInMB)
+			{
+				FString Text;
+				FLinearColor TextColor = MBTextColor;
+
+				if (Offset != 0)
+				{
+					Text = FString::Printf(TEXT("%0.2fMiB"), (float)Offset / MiB);
+				}
+				else
+				{
+					Text = FString::Printf(TEXT("Heap[%d] %uMiB"), HeapIndex, HeapSizeInMB);
+					TextColor.A = 1.0f;
+				}
+
+				const FSlateFontInfo& Font = DrawHelper.GetEventFont();
+				const TSharedRef<FSlateFontMeasure> FontMeasureService = FSlateApplication::Get().GetRenderer()->GetFontMeasureService();
+				const FVector2D TextExtent = FontMeasureService->Measure(Text, Font);
+
+				const float X = FMath::Max(VisibleGraph.SlateX - 2.0f - TextExtent.X, 2.0f);
+				const float Y = GetPosY() + GetLaneY(Layout, Depth) - TextExtent.Y * 0.5f;
+
+				DrawContext.DrawText(DrawHelper.GetFirstLayerId() + int32(EDrawLayer::Background), X, Y, Text, Font, TextColor);
+				return X + TextExtent.X + 2.0f;
+			};
+
+			const float SinglePixelDepth = 1.0f / Layout.EventH;
+
+			for (int32 HeapIndex = 0; HeapIndex < Graph.TransientHeapOffsets.Num(); ++HeapIndex)
+			{
+				const float HeapStartDepth = StartDepth + VisibleGraph.TransientHeapDepths[HeapIndex];
+
+				const auto& TransientHeap = Graph.TransientHeapStats.Heaps[HeapIndex];
+
+				uint64 Offset = 0;
+				uint32 OffsetIndex = 0;
+
+				const uint32 TicksPerText = 4;
+
+				while (Offset < TransientHeap.WatermarkSize)
+				{
+					const float Depth = HeapStartDepth + TransientHeapOffsetToDepth(Offset);
+
+					FLinearColor LineColor = MBLineColor;
 
-				const float X = FMath::Lerp(GraphicsForkVisiblePass.Min.X, GraphicsForkVisiblePass.Max.X, EndT);
-				const float Y = GraphicsForkVisiblePass.Max.Y;
-				const float LocalEndX = FMath::Lerp(AsyncComputeVisiblePass.Min.X, AsyncComputeVisiblePass.Max.X, StartT) - X;
-				const float LocalEndY = AsyncComputeVisiblePass.Min.Y - Y;
+					float LineStartX = VisibleGraph.SlateX;
+					float LineWidth  = VisibleGraph.SlateW;
 
-				FSpline Spline;
-				Spline.Start.X = X;
-				Spline.Start.Y = Y;
-				Spline.StartDir = FVector2D(0, SplineDir);
-				Spline.End = FVector2D(LocalEndX, LocalEndY);
-				Spline.EndDir = FVector2D(0, SplineDir);
-				Spline.Thickness = 2.0f;
-				Spline.Tint = FLinearColor(0.4f, 1.0f, 0.4f, TintAlpha);
-				DrawClampedSpline(SplineUnderLineLayerId + 1, Spline);
-			}
+					const bool bMajorTick = (OffsetIndex % TicksPerText) == 0;
 
-			const FRDGPassHandle GraphicsJoinHandle = AsyncComputePass.GraphicsJoinPass;
+					if (bMajorTick)
+					{
+						LineColor.A = 0.5f;
+					}
 
-			if (GraphicsJoinHandle.IsValid() && AsyncComputePass.bAsyncComputeEnd)
-			{
-				const FPassPacket& GraphicsJoinPass = *Graph.GetPass(AsyncComputePass.GraphicsJoinPass);
-				const FVisiblePass& GraphicsJoinVisiblePass = VisibleGraph.GetVisiblePass(GraphicsJoinPass);
+					if (bDrawText && bMajorTick)
+					{
+						LineStartX = DrawMibText(Depth, Offset, HeapIndex, TransientHeap.Capacity / MiB);
+						LineWidth -= LineStartX - VisibleGraph.SlateX;
+					}
 
-				const float X = FMath::Lerp(AsyncComputeVisiblePass.Min.X, AsyncComputeVisiblePass.Max.X, EndT);
-				const float Y = AsyncComputeVisiblePass.Min.Y;
-				const float LocalEndX = FMath::Lerp(GraphicsJoinVisiblePass.Min.X, GraphicsJoinVisiblePass.Max.X, StartT) - X;
-				const float LocalEndY = GraphicsJoinVisiblePass.Max.Y - Y;
+					DrawHelper.DrawBox(*this, LineStartX, Depth, LineWidth, SinglePixelDepth, LineColor, EDrawLayer::Background);
 
-				FSpline Spline;
-				Spline.Start.X = X;
-				Spline.Start.Y = Y;
-				Spline.StartDir = FVector2D(0, -SplineDir);
-				Spline.End = FVector2D(LocalEndX, LocalEndY);
-				Spline.EndDir = FVector2D(0, -SplineDir);
-				Spline.Thickness = 2.0f;
-				Spline.Tint = FLinearColor(1.0f, 0.4f, 0.4f, TintAlpha);
-				DrawClampedSpline(SplineUnderLineLayerId + 1, Spline);
+					Offset += TransientHeapBytesPerDepthSlot * (Layout.bIsCompactMode ? 3 : 1);
+					OffsetIndex++;
+				}
 			}
 		}
 	}
@@ -655,110 +916,190 @@ void FRenderGraphTrack::PostDraw(const ITimingTrackDrawContext& Context) const
 	SelectedTooltipState.Draw(Context.GetDrawContext());
 }
 
-void FRenderGraphTrack::BuildDrawState(ITimingEventsTrackDrawStateBuilder& Builder, const ITimingTrackUpdateContext& Context)
+void FRenderGraphTrack::BuildDrawState(FRenderGraphTrackDrawStateBuilder& Builder, const ITimingTrackUpdateContext& Context)
 {
-	// Reset draw state; it's built each frame.
+	const float kMinGraphPixels = 5.0f;
+	const float kMinPassMarginPixels = 5.0f;
+
+	const FRenderGraphProvider* RenderGraphProvider = SharedData.GetAnalysisSession().ReadProvider<FRenderGraphProvider>(FRenderGraphProvider::ProviderName);
+
+	if (!RenderGraphProvider)
 	{
-		Splines.Empty();
-		for (FVisibleGraph& VisibleGraph : VisibleGraphs)
+		return;
+	}
+
+	const FTimingTrackViewport& Viewport = Context.GetViewport();
+	const FTimingViewLayout& Layout = Viewport.GetLayout();
+
+	const auto DrawPasses = [&](double GraphStartTime, double GraphEndTime, FVisibleGraph& VisibleGraph, double& SinglePixelTimeMargin, uint32& DepthOffset)
+	{
+		SinglePixelTimeMargin = Viewport.GetDurationForViewportDX(1.0);
+
+		const FGraphPacket& GraphPacket = VisibleGraph.GetPacket();
+
+		if (Viewport.GetViewportDXForDuration(GraphPacket.NormalizedPassDuration) <= kMinPassMarginPixels)
 		{
-			VisibleGraph.Reset();
+			SinglePixelTimeMargin = 0.0;
 		}
-		ResetVisibleItemArray(VisibleGraphs);
 
-		PacketFilter = nullptr;
-		if (TSharedPtr<ITimingEventFilter> Filter = Context.GetEventFilter())
+		DepthOffset = 1;
+
+		for (uint32 ScopeIndex = 0, ScopeCount = GraphPacket.Scopes.Num(); ScopeIndex < ScopeCount; ++ScopeIndex)
 		{
-			if (Filter->Is<FPacketFilter>())
-			{
-				PacketFilter = &Filter->As<FPacketFilter>();
-			}
+			const FScopePacket& Scope = GraphPacket.Scopes[ScopeIndex];
+			const double StartTime = Scope.StartTime + SinglePixelTimeMargin;
+			const double EndTime = Scope.EndTime;
+
+			const FVisibleScope VisibleScope(Viewport, Scope, FTimingEvent::ComputeEventColor(*Scope.Name), StartTime, EndTime, DepthOffset + Scope.Depth);
+			VisibleScope.Draw(Builder);
+			VisibleGraph.AddScope(VisibleScope);
 		}
-	}
 
-	if (const FRenderGraphProvider* RenderGraphProvider = SharedData.GetAnalysisSession().ReadProvider<FRenderGraphProvider>(FRenderGraphProvider::ProviderName))
+		if (GraphPacket.ScopeDepth)
+		{
+			DepthOffset += GraphPacket.ScopeDepth + 1;
+		}
+
+		VisibleGraph.HeaderPassBarDepth = DepthOffset++;
+
+		bool bAnyAsyncCompute = false;
+
+		for (uint32 PassIndex = 0, PassCount = GraphPacket.Passes.Num(); PassIndex < PassCount; ++PassIndex)
+		{
+			const FPassPacket& Pass = GraphPacket.Passes[PassIndex];
+			const double StartTime = Pass.StartTime + SinglePixelTimeMargin;
+			const double EndTime = Pass.EndTime;
+			const bool bAsyncCompute = EnumHasAnyFlags(Pass.Flags, ERDGPassFlags::AsyncCompute);
+			const uint32 Depth = DepthOffset + (bAsyncCompute ? 2 : 0);
+			const uint32 Color = GetPassColor(Pass);
+
+			const FVisiblePass VisiblePass(Viewport, Pass, Color, StartTime, EndTime, Depth);
+			VisiblePass.Draw(Builder);
+			VisibleGraph.AddPass(VisiblePass);
+
+			bAnyAsyncCompute |= bAsyncCompute;
+		}
+
+		// Empty space between passes / resources.
+		DepthOffset += bAnyAsyncCompute ? 3 : 1;
+
+		VisibleGraph.HeaderEndDepth = DepthOffset;
+		VisibleGraph.DepthH = DepthOffset;
+	};
+
+	if (Visualizer == EVisualizer::TransientHeaps)
 	{
 		struct FResourceEntry
 		{
 			double StartTime{};
 			double EndTime{};
-			uint64 SizeInBytes{};
-			uint32 Index{};
-			uint32 Order{};
+			FRHITransientResourceStats TransientStats;
 			ERDGParentResourceType Type = ERDGParentResourceType::Texture;
-			bool bHasPreviousOwner{};
 		};
 
 		TArray<FResourceEntry> Resources;
-		TArray<uint16> TextureIndexToDepth;
-		TArray<uint16> BufferIndexToDepth;
-
-		const FTimingTrackViewport& Viewport = Context.GetViewport();
 
 		TraceServices::FAnalysisSessionReadScope SessionReadScope(SharedData.GetAnalysisSession());
 
 		const FRenderGraphProvider::TGraphTimeline& GraphTimeline = RenderGraphProvider->GetGraphTimeline();
 		GraphTimeline.EnumerateEvents(Viewport.GetStartTime(), Viewport.GetEndTime(),
-			[this, &Builder, &Viewport, &Resources, &TextureIndexToDepth, &BufferIndexToDepth](double GraphStartTime, double GraphEndTime, uint32, const TSharedPtr<FGraphPacket>& Graph)
+			[&](double GraphStartTime, double GraphEndTime, uint32, const TSharedPtr<FGraphPacket>& Graph)
 		{
-			// We always render the graph event, so add it separately even if the visible graph is culled.
-			Builder.AddEvent(GraphStartTime, GraphEndTime, 0, *Graph->Name, 0, kBuilderColor);
+			FVisibleGraph& VisibleGraph = AddVisibleGraph(Viewport, *Graph, kBuilderColor);
+			VisibleGraph.Draw(Builder);
 
 			if (Viewport.GetViewportDXForDuration(GraphEndTime - GraphStartTime) <= kMinGraphPixels)
 			{
 				return TraceServices::EEventEnumerate::Continue;
 			}
 
-			FVisibleGraph& VisibleGraph = AddVisibleGraph(Viewport, *Graph, kBuilderColor);
+			uint32 DepthOffset;
+			double SinglePixelTimeMargin;
+			DrawPasses(GraphStartTime, GraphEndTime, VisibleGraph, SinglePixelTimeMargin, DepthOffset);
 
-			double SinglePixelTimeMargin = Viewport.GetDurationForViewportDX(1.0);
+			const float YOffset = Layout.GetLaneY(DepthOffset);
 
-			if (Viewport.GetViewportDXForDuration(Graph->NormalizedPassDuration) <= kMinPassMarginPixels)
+			for (uint32 TextureIndex = 0, TextureCount = Graph->Textures.Num(); TextureIndex < TextureCount; ++TextureIndex)
 			{
-				SinglePixelTimeMargin = 0.0;
-			}
+				FTexturePacket& Texture = Graph->Textures[TextureIndex];
 
-			uint32 DepthOffset = 1;
+				if (!Texture.bTransient)
+				{
+					continue;
+				}
 
-			for (uint32 ScopeIndex = 0, ScopeCount = Graph->Scopes.Num(); ScopeIndex < ScopeCount; ++ScopeIndex)
-			{
-				const FScopePacket& Scope = Graph->Scopes[ScopeIndex];
-				const double StartTime = Scope.StartTime + SinglePixelTimeMargin;
-				const double EndTime = Scope.EndTime;
+				const float HeapOffsetDepthBase = VisibleGraph.TransientHeapDepths[Texture.TransientStats.HeapIndex];
 
-				const FVisibleScope VisibleScope(Viewport, Scope, StartTime, EndTime, DepthOffset + Scope.Depth, FTimingEvent::ComputeEventColor(*Scope.Name));
-				AddEvent(Builder, VisibleScope);
-				VisibleGraph.AddScope(VisibleScope);
+				const float DepthY = TransientHeapOffsetToDepth(Texture.TransientStats.HeapOffsetMin) + HeapOffsetDepthBase + DepthOffset;
+				const float DepthH = TransientHeapOffsetToDepth(Texture.TransientStats.HeapOffsetMax - Texture.TransientStats.HeapOffsetMin);
+				const double StartTime = Texture.StartTime + SinglePixelTimeMargin;
+				const double EndTime = Texture.EndTime;
+				const uint32 Color = kTextureColor;
+
+				const FVisibleTexture VisibleTexture(Viewport, Texture, Color, StartTime, EndTime, DepthY, DepthH);
+				VisibleTexture.Draw(Builder);
+				VisibleGraph.AddTexture(VisibleTexture);
 			}
 
-			if (Graph->ScopeDepth)
+			for (uint32 BufferIndex = 0, BufferCount = Graph->Buffers.Num(); BufferIndex < BufferCount; ++BufferIndex)
 			{
-				DepthOffset += Graph->ScopeDepth + 1;
+				FBufferPacket& Buffer = Graph->Buffers[BufferIndex];
+
+				if (!Buffer.bTransient)
+				{
+					continue;
+				}
+				
+				const float HeapOffsetDepthBase = VisibleGraph.TransientHeapDepths[Buffer.TransientStats.HeapIndex];
+
+				const float DepthY = TransientHeapOffsetToDepth(Buffer.TransientStats.HeapOffsetMin) + HeapOffsetDepthBase + DepthOffset;
+				const float DepthH = TransientHeapOffsetToDepth(Buffer.TransientStats.HeapOffsetMax - Buffer.TransientStats.HeapOffsetMin);
+				const double StartTime = Buffer.StartTime + SinglePixelTimeMargin;
+				const double EndTime = Buffer.EndTime;
+				const uint32 Color = kBufferColor;
+
+				const FVisibleBuffer VisibleBuffer(Viewport, Buffer, Color, StartTime, EndTime, DepthY, DepthH);
+				VisibleBuffer.Draw(Builder);
+				VisibleGraph.AddBuffer(VisibleBuffer);
 			}
 
-			// +1 for render pass merge bars.
-			DepthOffset += 1;
+			return TraceServices::EEventEnumerate::Continue;
+		});
+	}
+	else if (Visualizer == EVisualizer::Resources)
+	{
+		struct FResourceEntry
+		{
+			double StartTime{};
+			double EndTime{};
+			uint64 SizeInBytes{};
+			uint32 Index{};
+			uint32 Order{};
+			ERDGParentResourceType Type = ERDGParentResourceType::Texture;
+			bool bHasPreviousOwner{};
+		};
+
+		TArray<FResourceEntry> Resources;
+		TArray<uint16> TextureIndexToDepth;
+		TArray<uint16> BufferIndexToDepth;
+
+		TraceServices::FAnalysisSessionReadScope SessionReadScope(SharedData.GetAnalysisSession());
 
-			bool bAnyAsyncCompute = false;
+		const FRenderGraphProvider::TGraphTimeline& GraphTimeline = RenderGraphProvider->GetGraphTimeline();
+		GraphTimeline.EnumerateEvents(Viewport.GetStartTime(), Viewport.GetEndTime(),
+			[&](double GraphStartTime, double GraphEndTime, uint32, const TSharedPtr<FGraphPacket>& Graph)
+		{
+			FVisibleGraph& VisibleGraph = AddVisibleGraph(Viewport, *Graph, kBuilderColor);
+			VisibleGraph.Draw(Builder);
 
-			for (uint32 PassIndex = 0, PassCount = Graph->Passes.Num(); PassIndex < PassCount; ++PassIndex)
+			if (Viewport.GetViewportDXForDuration(GraphEndTime - GraphStartTime) <= kMinGraphPixels)
 			{
-				FPassPacket& Pass = Graph->Passes[PassIndex];
-				const double StartTime = Pass.StartTime + SinglePixelTimeMargin;
-				const double EndTime = Pass.EndTime;
-				const bool bAsyncCompute = EnumHasAnyFlags(Pass.Flags, ERDGPassFlags::AsyncCompute);
-				const uint32 Depth = DepthOffset + (bAsyncCompute ? 2 : 0);
-				const uint32 Color = GetPassColor(Pass);
-
-				const FVisiblePass VisiblePass(Viewport, Pass, StartTime, EndTime, Depth, Color);
-				AddEvent(Builder, VisiblePass);
-				VisibleGraph.AddPass(VisiblePass);
-
-				bAnyAsyncCompute |= bAsyncCompute;
+				return TraceServices::EEventEnumerate::Continue;
 			}
 
-			// Empty space between passes / resources.
-			DepthOffset += bAnyAsyncCompute ? 3 : 1;
+			uint32 DepthOffset;
+			double SinglePixelTimeMargin;
+			DrawPasses(GraphStartTime, GraphEndTime, VisibleGraph, SinglePixelTimeMargin, DepthOffset);
 
 			Resources.Reset();
 			Resources.Reserve(Graph->Textures.Num() + Graph->Buffers.Num());
@@ -922,8 +1263,8 @@ void FRenderGraphTrack::BuildDrawState(ITimingEventsTrackDrawStateBuilder& Build
 					const double EndTime = Texture.EndTime;
 					const uint32 Color = GetTextureColor(Texture, MaxSizeInBytes);
 
-					const FVisibleTexture VisibleTexture(Viewport, Texture, StartTime, EndTime, Depth, Color);
-					AddEvent(Builder, VisibleTexture);
+					const FVisibleTexture VisibleTexture(Viewport, Texture, Color, StartTime, EndTime, Depth);
+					VisibleTexture.Draw(Builder);
 					VisibleGraph.AddTexture(VisibleTexture);
 				}
 			}
@@ -949,8 +1290,8 @@ void FRenderGraphTrack::BuildDrawState(ITimingEventsTrackDrawStateBuilder& Build
 					const double EndTime = Buffer.EndTime;
 					const uint32 Color = GetBufferColor(Buffer, MaxSizeInBytes);
 
-					const FVisibleBuffer VisibleBuffer(Viewport, Buffer, StartTime, EndTime, Depth, Color);
-					AddEvent(Builder, VisibleBuffer);
+					const FVisibleBuffer VisibleBuffer(Viewport, Buffer, Color, StartTime, EndTime, Depth);
+					VisibleBuffer.Draw(Builder);
 					VisibleGraph.AddBuffer(VisibleBuffer);
 				}
 			}
@@ -960,8 +1301,10 @@ void FRenderGraphTrack::BuildDrawState(ITimingEventsTrackDrawStateBuilder& Build
 	}
 }
 
-void FRenderGraphTrack::BuildFilteredDrawState(ITimingEventsTrackDrawStateBuilder& Builder, const ITimingTrackUpdateContext& Context)
+void FRenderGraphTrack::BuildFilteredDrawState(FRenderGraphTrackDrawStateBuilder& Builder, const ITimingTrackUpdateContext& Context)
 {
+	const FPacketFilter* PacketFilter = GetPacketFilter(Context.GetEventFilter());
+
 	if (!PacketFilter)
 	{
 		return;
@@ -975,13 +1318,13 @@ void FRenderGraphTrack::BuildFilteredDrawState(ITimingEventsTrackDrawStateBuilde
 		return;
 	}
 
-	AddEvent(Builder, *VisibleGraph);
+	VisibleGraph->Draw(Builder);
 
 	TSet<const FVisibleItem*> VisibleItems;
 
 	for (const FVisibleScope& VisibleScope : VisibleGraph->Scopes)
 	{
-		if (!FilterPacket(VisibleScope.GetPacket()))
+		if (!PacketFilter->FilterPacket(VisibleScope.GetPacket()))
 		{
 			continue;
 		}
@@ -993,7 +1336,7 @@ void FRenderGraphTrack::BuildFilteredDrawState(ITimingEventsTrackDrawStateBuilde
 	{
 		const FPassPacket& Pass = VisiblePass.GetPacket();
 
-		if (!FilterPacket(Pass))
+		if (!PacketFilter->FilterPacket(Pass))
 		{
 			continue;
 		}
@@ -1051,27 +1394,30 @@ void FRenderGraphTrack::BuildFilteredDrawState(ITimingEventsTrackDrawStateBuilde
 		}
 	}
 
+	const FTimingViewLayout& Layout = Context.GetViewport().GetLayout();
+
 	const auto AddResourcePassEvents = [&](const FVisibleItem& VisibleResource, TArrayView<const FRDGPassHandle> Passes)
 	{
-		AddEvent(Builder, VisibleResource);
+		VisibleResource.Draw(Builder);
 
 		for (FRDGPassHandle PassHandle : Passes)
 		{
 			const FPassPacket& Pass = *Graph.GetPass(PassHandle);
 			const FVisiblePass& VisiblePass = VisibleGraph->GetVisiblePass(Pass);
-			AddEvent(Builder, VisiblePass);
 
-			const float Y = VisiblePass.Max.Y;
+			VisiblePass.Draw(Builder);
 
-			FSpline Spline;
-			Spline.Start.X = (VisiblePass.Max.X + VisiblePass.Min.X) * 0.5f;
-			Spline.Start.Y = Y;
+			const float DepthY = VisiblePass.DepthY + VisiblePass.DepthH;
+
+			FSplinePrimitive Spline;
+			Spline.Start.X = VisiblePass.SlateX + (VisiblePass.SlateW * 0.5f);
+			Spline.Start.Y = DepthY;
 			Spline.StartDir = FVector2D(0, -1);
-			Spline.End = FVector2D(0, VisibleResource.Min.Y - Y);
+			Spline.End = FVector2D(0, VisibleResource.DepthY - DepthY);
 			Spline.EndDir = FVector2D(0, 1);
 			Spline.Thickness = 1.0f;
 			Spline.Tint = FLinearColor(0.8f, 0.8f, 0.8f, 0.7f);
-			Splines.Add(Spline);
+			Builder.AddSpline(Spline);
 
 			for (const FVisibleScope& VisibleScope : VisibleGraph->Scopes)
 			{
@@ -1090,7 +1436,7 @@ void FRenderGraphTrack::BuildFilteredDrawState(ITimingEventsTrackDrawStateBuilde
 		{
 			const FTexturePacket& Texture = VisibleTexture.GetPacket();
 
-			if (FilterPacketExact(Texture))
+			if (PacketFilter->FilterPacketExact(Texture))
 			{
 				AddResourcePassEvents(VisibleTexture, Texture.Passes);
 			}
@@ -1103,7 +1449,7 @@ void FRenderGraphTrack::BuildFilteredDrawState(ITimingEventsTrackDrawStateBuilde
 		{
 			const FBufferPacket& Buffer = VisibleBuffer.GetPacket();
 
-			if (FilterPacketExact(Buffer))
+			if (PacketFilter->FilterPacketExact(Buffer))
 			{
 				AddResourcePassEvents(VisibleBuffer, Buffer.Passes);
 			}
@@ -1112,21 +1458,33 @@ void FRenderGraphTrack::BuildFilteredDrawState(ITimingEventsTrackDrawStateBuilde
 
 	for (const FVisibleItem* VisibleItem : VisibleItems)
 	{
-		AddEvent(Builder, *VisibleItem);
+		VisibleItem->Draw(Builder);
+	}
+}
+
+void FRenderGraphTrack::DrawEvent(const ITimingTrackDrawContext& Context, const ITimingEvent& InTimingEvent, EDrawEventMode InDrawMode) const
+{
+	if (InTimingEvent.CheckTrack(this) && InTimingEvent.Is<FRenderGraphTimingEvent>())
+	{
+		const FRenderGraphTimingEvent& TrackEvent = InTimingEvent.As<FRenderGraphTimingEvent>();
+		const FVisibleItem TrackItem = TrackEvent.GetItem();
+
+		FRenderGraphTrackDrawHelper DrawHelper(Context);
+		DrawHelper.DrawTimingEventHighlight(*this, TrackEvent.GetStartTime(), TrackEvent.GetEndTime(), TrackItem.DepthY, TrackItem.DepthH, InDrawMode);
 	}
 }
 
 const TSharedPtr<const ITimingEvent> FRenderGraphTrack::GetEvent(float InPosX, float InPosY, const FTimingTrackViewport& Viewport) const
 {
-	TSharedPtr<FEvent> Event;
+	TSharedPtr<FRenderGraphTimingEvent> Event;
+
+	const float DepthY = GetDepthFromLaneY(Viewport.GetLayout(), InPosY - GetPosY());
 
 	for (const FVisibleGraph& Graph : VisibleGraphs)
 	{
-		float AdjustedPosY = InPosY - GetPosY();
-
-		if (Graph.Intersects(InPosX, AdjustedPosY))
+		if (Graph.Intersects(InPosX, DepthY))
 		{
-			if (const FVisibleItem* Item = Graph.FindItem(InPosX, AdjustedPosY))
+			if (const FVisibleItem* Item = Graph.FindItem(InPosX, DepthY))
 			{
 				const FPacket& Packet = *Item->Packet;
 
@@ -1155,12 +1513,12 @@ const TSharedPtr<const ITimingEvent> FRenderGraphTrack::GetEvent(float InPosX, f
 
 TSharedPtr<ITimingEventFilter> FRenderGraphTrack::GetFilterByEvent(const TSharedPtr<const ITimingEvent> InTimingEvent) const
 {
-	if (InTimingEvent.IsValid() && InTimingEvent->Is<FEvent>())
+	if (InTimingEvent.IsValid() && InTimingEvent->Is<FRenderGraphTimingEvent>())
 	{
 		bool bFilterable = false;
-		if (InTimingEvent->As<FEvent>().GetItem().Intersects(MousePosition.X, MousePosition.Y, bFilterable) && bFilterable)
+		if (InTimingEvent->As<FRenderGraphTimingEvent>().GetItem().Intersects(MouseSlateX, MouseDepthY, bFilterable) && bFilterable)
 		{
-			return MakeShared<FPacketFilter>(StaticCastSharedPtr<const FEvent>(InTimingEvent));
+			return MakeShared<FPacketFilter>(StaticCastSharedPtr<const FRenderGraphTimingEvent>(InTimingEvent));
 		}
 	}
 	return nullptr;
@@ -1188,8 +1546,57 @@ void FRenderGraphTrack::BuildContextMenu(FMenuBuilder& MenuBuilder)
 {
 	Super::BuildContextMenu(MenuBuilder);
 
-	MenuBuilder.BeginSection("Show", LOCTEXT("ShowMenuHeader", "Track Show Flags"));
+	Insights::ITimingViewSession* TimingViewSession = SharedData.GetTimingViewSession();
+
+	MenuBuilder.BeginSection("Visualizer", LOCTEXT("Visualizer", "Visualizer Mode"));
+	{
+		MenuBuilder.AddMenuEntry
+		(
+			LOCTEXT("Resources", "Resources"),
+			LOCTEXT("Resources_Tooltip", "Show each resource on its own line."),
+			FSlateIcon(),
+			FUIAction(
+				FExecuteAction::CreateLambda([this, TimingViewSession]()
+				{
+					Visualizer = EVisualizer::Resources;
+					TimingViewSession->ResetSelectedEvent();
+					TimingViewSession->ResetEventFilter();
+					SetDirtyFlag();
+				}),
+				FCanExecuteAction(),
+				FIsActionChecked::CreateLambda([this, TimingViewSession]() { return Visualizer == EVisualizer::Resources; })
+			),
+			NAME_None,
+			EUserInterfaceActionType::RadioButton
+		);
+
+		MenuBuilder.AddMenuEntry
+		(
+			LOCTEXT("TransientHeap", "Transient Heaps"),
+			LOCTEXT("ShowTextures_Tooltip", "Shows resource placement in transient heaps."),
+			FSlateIcon(),
+			FUIAction(
+				FExecuteAction::CreateLambda([this, TimingViewSession]()
+				{
+					Visualizer = EVisualizer::TransientHeaps;
+					TimingViewSession->ResetSelectedEvent();
+					TimingViewSession->ResetEventFilter();
+					ResourceShow = EResourceShow::All;
+					SetDirtyFlag();
+				}),
+				FCanExecuteAction(),
+				FIsActionChecked::CreateLambda([this]() { return Visualizer == EVisualizer::TransientHeaps; })
+			),
+			NAME_None,
+			EUserInterfaceActionType::RadioButton
+		);
+	}
+	MenuBuilder.EndSection();
+
+	if (Visualizer == EVisualizer::Resources)
 	{
+		MenuBuilder.BeginSection("Show", LOCTEXT("ShowMenuHeader", "Track Show Flags"));
+
 		MenuBuilder.AddMenuEntry
 		(
 			LOCTEXT("ShowTextures", "Show Textures"),
@@ -1279,11 +1686,10 @@ void FRenderGraphTrack::BuildContextMenu(FMenuBuilder& MenuBuilder)
 			NAME_None,
 			EUserInterfaceActionType::ToggleButton
 		);
-	}
-	MenuBuilder.EndSection();
 
-	MenuBuilder.BeginSection("Sort", LOCTEXT("SortMenuHeader", "Track Sort By"));
-	{
+		MenuBuilder.EndSection();
+		MenuBuilder.BeginSection("Sort", LOCTEXT("SortMenuHeader", "Track Sort By"));
+
 		MenuBuilder.AddMenuEntry
 		(
 			LOCTEXT("SortCreation", "Creation"),
@@ -1373,11 +1779,10 @@ void FRenderGraphTrack::BuildContextMenu(FMenuBuilder& MenuBuilder)
 			NAME_None,
 			EUserInterfaceActionType::RadioButton
 		);
-	}
-	MenuBuilder.EndSection();
 
-	MenuBuilder.BeginSection("Color", LOCTEXT("ColorMenuHeader", "Track Resource Coloration"));
-	{
+		MenuBuilder.EndSection();
+		MenuBuilder.BeginSection("Color", LOCTEXT("ColorMenuHeader", "Track Resource Coloration"));
+
 		MenuBuilder.AddMenuEntry
 		(
 			LOCTEXT("ColorType", "By Type"),
@@ -1413,11 +1818,10 @@ void FRenderGraphTrack::BuildContextMenu(FMenuBuilder& MenuBuilder)
 			NAME_None,
 			EUserInterfaceActionType::RadioButton
 		);
-	}
-	MenuBuilder.EndSection();
 
-	MenuBuilder.BeginSection("FilterText", LOCTEXT("FilterTextHeader", "Track Resource Filter"));
-	{
+		MenuBuilder.EndSection();
+		MenuBuilder.BeginSection("FilterText", LOCTEXT("FilterTextHeader", "Track Resource Filter"));
+
 		MenuBuilder.AddWidget(
 			SNew(SVerticalBox)
 			+ SVerticalBox::Slot()
@@ -1454,19 +1858,164 @@ void FRenderGraphTrack::BuildContextMenu(FMenuBuilder& MenuBuilder)
 					.OnValueCommitted_Lambda([this](float InValue, ETextCommit::Type) { FilterSize = InValue; SetDirtyFlag(); })
 				]
 			],
-			FText::GetEmpty(), true);
+			FText::GetEmpty(),
+			true
+		);
+	
+		MenuBuilder.EndSection();
+	}
+	else if (Visualizer == EVisualizer::TransientHeaps)
+	{
+		const uint64 KiB = 1024;
+		const uint64 MiB = 1024 * KiB;
+
+		MenuBuilder.BeginSection("HeapZoom", LOCTEXT("HeapZoomHeader", "Vertical Zoom"));
+		
+		MenuBuilder.AddMenuEntry
+		(
+			LOCTEXT("4MiB", "4MiB"),
+			LOCTEXT("4MiB_Tooltip", "Each depth slot in the timing view is 1MiB in size on the vertical axis."),
+			FSlateIcon(),
+			FUIAction(
+				FExecuteAction::CreateLambda([this, MiB]()
+				{
+					TransientHeapBytesPerDepthSlot = 4 * MiB;
+					SetDirtyFlag();
+				}),
+				FCanExecuteAction(),
+				FIsActionChecked::CreateLambda([this, MiB]() { return TransientHeapBytesPerDepthSlot == 4 * MiB; })
+			),
+			NAME_None,
+			EUserInterfaceActionType::RadioButton
+		);
+		
+		MenuBuilder.AddMenuEntry
+		(
+			LOCTEXT("2MiB", "2MiB"),
+			LOCTEXT("2MiB_Tooltip", "Each depth slot in the timing view is 1MiB in size on the vertical axis."),
+			FSlateIcon(),
+			FUIAction(
+				FExecuteAction::CreateLambda([this, MiB]()
+				{
+					TransientHeapBytesPerDepthSlot = 2 * MiB;
+					SetDirtyFlag();
+				}),
+				FCanExecuteAction(),
+				FIsActionChecked::CreateLambda([this, MiB]() { return TransientHeapBytesPerDepthSlot == 2 * MiB; })
+			),
+			NAME_None,
+			EUserInterfaceActionType::RadioButton
+		);
+		
+		MenuBuilder.AddMenuEntry
+		(
+			LOCTEXT("1MiB", "1MiB"),
+			LOCTEXT("1MiB_Tooltip", "Each depth slot in the timing view is 1MiB in size on the vertical axis."),
+			FSlateIcon(),
+			FUIAction(
+				FExecuteAction::CreateLambda([this, MiB]()
+				{
+					TransientHeapBytesPerDepthSlot = MiB;
+					SetDirtyFlag();
+				}),
+				FCanExecuteAction(),
+				FIsActionChecked::CreateLambda([this, MiB]() { return TransientHeapBytesPerDepthSlot == MiB; })
+			),
+			NAME_None,
+			EUserInterfaceActionType::RadioButton
+		);
+		
+		MenuBuilder.AddMenuEntry
+		(
+			LOCTEXT("512KiB", "512KiB"),
+			LOCTEXT("512KiB_Tooltip", "Each depth slot in the timing view is 512KiB in size on the vertical axis."),
+			FSlateIcon(),
+			FUIAction(
+				FExecuteAction::CreateLambda([this, KiB]()
+				{
+					TransientHeapBytesPerDepthSlot = 512 * KiB;
+					SetDirtyFlag();
+				}),
+				FCanExecuteAction(),
+				FIsActionChecked::CreateLambda([this, KiB]() { return TransientHeapBytesPerDepthSlot == 512 * KiB; })
+			),
+			NAME_None,
+			EUserInterfaceActionType::RadioButton
+		);
+
+		MenuBuilder.AddMenuEntry
+		(
+			LOCTEXT("256KiB", "256KiB"),
+			LOCTEXT("256KiB_Tooltip", "Each depth slot in the timing view is 256KiB in size on the vertical axis."),
+			FSlateIcon(),
+			FUIAction(
+				FExecuteAction::CreateLambda([this, KiB]()
+				{
+					TransientHeapBytesPerDepthSlot = 256 * KiB;
+					SetDirtyFlag();
+				}),
+				FCanExecuteAction(),
+				FIsActionChecked::CreateLambda([this, KiB]() { return TransientHeapBytesPerDepthSlot == 256 * KiB; })
+			),
+			NAME_None,
+			EUserInterfaceActionType::RadioButton
+		);
+		
+		MenuBuilder.AddMenuEntry
+		(
+			LOCTEXT("128KiB", "128KiB"),
+			LOCTEXT("128KiB_Tooltip", "Each depth slot in the timing view is 128KiB in size on the vertical axis."),
+			FSlateIcon(),
+			FUIAction(
+				FExecuteAction::CreateLambda([this, KiB]()
+				{
+					TransientHeapBytesPerDepthSlot = 128 * KiB;
+					SetDirtyFlag();
+				}),
+				FCanExecuteAction(),
+				FIsActionChecked::CreateLambda([this, KiB]() { return TransientHeapBytesPerDepthSlot == 128 * KiB; })
+			),
+			NAME_None,
+			EUserInterfaceActionType::RadioButton
+		);
+
+		MenuBuilder.AddMenuEntry
+		(
+			LOCTEXT("64KiB", "64KiB"),
+			LOCTEXT("64KiB_Tooltip", "Each depth slot in the timing view is 64KiB in size on the vertical axis."),
+			FSlateIcon(),
+			FUIAction(
+				FExecuteAction::CreateLambda([this, KiB]()
+				{
+					TransientHeapBytesPerDepthSlot = 64 * KiB;
+					SetDirtyFlag();
+				}),
+				FCanExecuteAction(),
+				FIsActionChecked::CreateLambda([this, KiB]() { return TransientHeapBytesPerDepthSlot == 64 * KiB; })
+			),
+			NAME_None,
+			EUserInterfaceActionType::RadioButton
+		);
+
+		MenuBuilder.EndSection();
 	}
-	MenuBuilder.EndSection();
 }
 
 void FRenderGraphTrack::InitTooltip(FTooltipDrawState& Tooltip, const ITimingEvent& InTooltipEvent) const
 {
 	Tooltip.ResetContent();
 
-	if (InTooltipEvent.CheckTrack(this) && InTooltipEvent.Is<FEvent>())
+	if (InTooltipEvent.CheckTrack(this) && InTooltipEvent.Is<FRenderGraphTimingEvent>())
 	{
 		const auto AddCommonResourceText = [&](const FResourcePacket& Resource)
 		{
+			if (Resource.TransientStats.HeapOffsetMax > 0)
+			{
+				Tooltip.AddNameValueTextLine(TEXT("Heap Index:"), FString::Printf(TEXT("%u"), Resource.TransientStats.HeapIndex));
+				Tooltip.AddNameValueTextLine(TEXT("Heap Start:"), FString::Printf(TEXT("%.3fMB"), (double)Resource.TransientStats.HeapOffsetMin / (1024.0 * 1024.0)));
+				Tooltip.AddNameValueTextLine(TEXT("Heap End:"), FString::Printf(TEXT("%.3fMB"), (double)Resource.TransientStats.HeapOffsetMax / (1024.0 * 1024.0)));
+			}
+
 			if (Resource.bExtracted)
 			{
 				Tooltip.AddTextLine(TEXT("Extracted"), FLinearColor::Red);
@@ -1559,7 +2108,6 @@ void FRenderGraphTrack::InitTooltip(FTooltipDrawState& Tooltip, const ITimingEve
 	}
 }
 
-} //namespace RenderGraphInsights
-} //namespace UE
+}} //namespace UE::RenderGraphInsights
 
 #undef LOCTEXT_NAMESPACE
\ No newline at end of file
diff --git a/Engine/Plugins/RenderGraphInsights/Source/RenderGraphInsights/Private/RenderGraphTrack.h b/Engine/Plugins/RenderGraphInsights/Source/RenderGraphInsights/Private/RenderGraphTrack.h
--- a/Engine/Plugins/RenderGraphInsights/Source/RenderGraphInsights/Private/RenderGraphTrack.h
+++ b/Engine/Plugins/RenderGraphInsights/Source/RenderGraphInsights/Private/RenderGraphTrack.h
@@ -7,59 +7,74 @@
 #include "Insights/ViewModels/TimingTrackViewport.h"
 #include "Insights/ViewModels/TooltipDrawState.h"
 #include "RenderGraphProvider.h"
+#include "RenderGraphTrackDrawHelper.h"
 
 class FMenuBuilder;
 class ITimingTrackDrawContext;
 class ITimingTrackUpdateContext;
 class FTimingTrackViewport;
 
-namespace UE
-{
-namespace RenderGraphInsights
-{
-class FRenderGraphTimingViewSession;
+namespace UE { namespace RenderGraphInsights {
 
-class FGraphPacket;
+class FRenderGraphTimingViewSession;
 class FPacketFilter;
 
+class FBoundingBox
+{
+public:
+	FBoundingBox() = default;
+
+	FBoundingBox(FVector2D InMin, FVector2D InMax)
+		: Min(InMin)
+		, Max(InMax)
+	{}
+
+	FVector2D Min = FVector2D::ZeroVector;
+	FVector2D Max = FVector2D::ZeroVector;
+};
+
 class FVisibleItem
 {
 public:
-	FVisibleItem(const FTimingTrackViewport& InViewport, const FPacket& InPacket, double InStartTime, double InEndTime, uint32 InDepth, uint32 InColor, float InMinY, float InMaxY)
+	FVisibleItem(const FTimingTrackViewport& InViewport, const FPacket& InPacket, uint32 InColor, double InStartTime, double InEndTime, float InDepthY, float InDepthH = 1.0f)
 		: Packet(&InPacket)
 		, Name(*InPacket.Name)
 		, StartTime(InStartTime)
 		, EndTime(InEndTime)
-		, Depth(InDepth)
+		, DepthY(InDepthY)
+		, DepthH(InDepthH)
 		, Color(InColor)
 	{
-		Min.X = InViewport.TimeToSlateUnitsRounded(StartTime);
-		Min.Y = InMinY;
-		Max.X = InViewport.TimeToSlateUnitsRounded(EndTime);
-		Max.Y = InMaxY;
+		SlateX = InViewport.TimeToSlateUnitsRounded(StartTime);
+		SlateW = InViewport.TimeToSlateUnitsRounded(EndTime) - SlateX;
 	}
 
-	virtual bool Intersects(float InPosX, float InPosY, bool& bOutFilterable) const
+	virtual bool Intersects(float InSlateX, float InDepthY, bool& bOutFilterable) const
 	{
 		bOutFilterable = true;
-		return InPosX >= Min.X && InPosX < Max.X&& InPosY > Min.Y && InPosY < Max.Y;
+		return InSlateX >= SlateX && InSlateX < (SlateX + SlateW) && InDepthY > DepthY && InDepthY < (DepthY + DepthH);
 	}
 
-	bool Intersects(float InPosX, float InPosY) const
+	bool Intersects(float InSlateX, float InDepthY) const
 	{
 		bool bFilterable = false;
-		return Intersects(InPosX, InPosY, bFilterable);
+		return Intersects(InSlateX, InDepthY, bFilterable);
 	}
 
+	FBoundingBox GetBoundingBox(const FTimingViewLayout& Layout) const;
+
+	virtual void Draw(FRenderGraphTrackDrawStateBuilder& Builder) const;
+
 	const FPacket* Packet{};
 	const TCHAR* Name{};
 	double StartTime{};
 	double EndTime{};
-	uint32 Depth{};
+	float DepthY{};
+	float DepthH{};
+	float SlateX{};
+	float SlateW{};
 	uint32 Color{};
 	uint32 Index{};
-	FVector2D Min = FVector2D::ZeroVector;
-	FVector2D Max = FVector2D::ZeroVector;
 };
 
 template <typename InPacketType>
@@ -68,16 +83,8 @@ class TVisibleItemHelper : public FVisibleItem
 public:
 	using PacketType = InPacketType;
 
-	TVisibleItemHelper(const FTimingTrackViewport& Viewport, const PacketType& InPacket, double InStartTime, double InEndTime, uint32 InDepth, uint32 InColor)
-		: TVisibleItemHelper(Viewport, InPacket, InStartTime, InEndTime, InDepth, InColor, Viewport.GetLayout().GetLaneY(InDepth))
-	{}
-
-	TVisibleItemHelper(const FTimingTrackViewport& Viewport, const PacketType& InPacket, double InStartTime, double InEndTime, uint32 InDepth, uint32 InColor, float InMinY)
-		: FVisibleItem(Viewport, InPacket, InStartTime, InEndTime, InDepth, InColor, InMinY, InMinY + Viewport.GetLayout().EventH)
-	{}
-
-	TVisibleItemHelper(const FTimingTrackViewport& Viewport, const PacketType& InPacket, double InStartTime, double InEndTime, uint32 InDepth, uint32 InColor, float InMinY, float InMaxY)
-		: FVisibleItem(Viewport, InPacket, InStartTime, InEndTime, InDepth, InColor, InMinY, InMaxY)
+	TVisibleItemHelper(const FTimingTrackViewport& Viewport, const PacketType& InPacket, uint32 InColor, double InStartTime, double InEndTime, float InDepthY, float InDepthH = 1.0f)
+		: FVisibleItem(Viewport, InPacket, InColor, InStartTime, InEndTime, InDepthY, InDepthH)
 	{}
 
 	const PacketType& GetPacket() const
@@ -89,16 +96,16 @@ public:
 class FVisiblePass final : public TVisibleItemHelper<FPassPacket>
 {
 public:
-	FVisiblePass(const FTimingTrackViewport& Viewport, const FPassPacket& InPacket, double InStartTime, double InEndTime, uint32 InDepth, uint32 InColor)
-		: TVisibleItemHelper(Viewport, InPacket, InStartTime, InEndTime, InDepth, InColor)
+	FVisiblePass(const FTimingTrackViewport& Viewport, const FPassPacket& InPacket, uint32 InColor, double InStartTime, double InEndTime, float InDepthY, float InDepthH = 1.0f)
+		: TVisibleItemHelper(Viewport, InPacket, InColor, InStartTime, InEndTime, InDepthY, InDepthH)
 	{}
 
-	bool Intersects(float InPosX, float InPosY, bool& bOutFilterable) const override
+	bool Intersects(float InSlateX, float InDepthY, bool& bOutFilterable) const override
 	{
 		// For hit-testing, treat the pass as unbounded along Y, so that it's a column.
-		if (InPosX >= Min.X && InPosX < Max.X && InPosY >= Min.Y)
+		if (InSlateX >= SlateX && InSlateX < (SlateX + SlateW) && InDepthY >= DepthY)
 		{
-			bOutFilterable = InPosY < Max.Y;
+			bOutFilterable = InDepthY < (DepthY + DepthH);
 			return true;
 		}
 		return false;
@@ -107,38 +114,15 @@ public:
 	using FVisibleItem::Intersects;
 };
 
-class FVisibleScope final : public TVisibleItemHelper<FScopePacket>
-{
-public:
-
-	FVisibleScope(const FTimingTrackViewport& Viewport, const FScopePacket& InPacket, double InStartTime, double InEndTime, uint32 InDepth, uint32 InColor)
-		: TVisibleItemHelper(Viewport, InPacket, InStartTime, InEndTime, InDepth, InColor)
-	{}
-};
-
-class FVisibleTexture final : public TVisibleItemHelper<FTexturePacket>
-{
-public:
-
-	FVisibleTexture(const FTimingTrackViewport& Viewport, const FTexturePacket& InPacket, double InStartTime, double InEndTime, uint32 InDepth, uint32 InColor)
-		: TVisibleItemHelper(Viewport, InPacket, InStartTime, InEndTime, InDepth, InColor)
-	{}
-};
-
-class FVisibleBuffer final : public TVisibleItemHelper<FBufferPacket>
-{
-public:
-
-	FVisibleBuffer(const FTimingTrackViewport& Viewport, const FBufferPacket& InPacket, double InStartTime, double InEndTime, uint32 InDepth, uint32 InColor)
-		: TVisibleItemHelper(Viewport, InPacket, InStartTime, InEndTime, InDepth, InColor)
-	{}
-};
+using FVisibleScope   = TVisibleItemHelper<FScopePacket>;
+using FVisibleTexture = TVisibleItemHelper<FTexturePacket>;
+using FVisibleBuffer  = TVisibleItemHelper<FBufferPacket>;
 
 class FVisibleGraph final : public TVisibleItemHelper<FGraphPacket>
 {
 public:
-	FVisibleGraph(const FTimingTrackViewport& Viewport, const FGraphPacket& InGraph, const uint32 InColor, float InMaxY)
-		: TVisibleItemHelper(Viewport, InGraph, InGraph.StartTime, InGraph.EndTime, 0, InColor, 0.0f, InMaxY)
+	FVisibleGraph(const FTimingTrackViewport& Viewport, const FGraphPacket& InGraph, uint32 InColor, float InDepthH)
+		: TVisibleItemHelper(Viewport, InGraph, InColor, InGraph.StartTime, InGraph.EndTime, 0.0f, InDepthH)
 	{}
 
 	void AddScope(const FVisibleScope& VisibleScope);
@@ -148,7 +132,7 @@ public:
 
 	void Reset();
 
-	const FVisibleItem* FindItem(float InPosX, float InPosY) const;
+	const FVisibleItem* FindItem(float InSlateX, float InDepthY) const;
 
 	const FVisibleScope& GetVisibleScope(const FScopePacket& Scope) const
 	{
@@ -180,43 +164,64 @@ public:
 		return &Buffers[Buffer.VisibleIndex];
 	}
 
+	void Draw(FRenderGraphTrackDrawStateBuilder& Builder) const override;
+
+	bool IsCulled() const
+	{
+		return Passes.IsEmpty();
+	}
+
 	TArray<FVisibleScope> Scopes;
 	TArray<FVisiblePass> Passes;
 	TArray<FVisibleTexture> Textures;
 	TArray<FVisibleBuffer> Buffers;
 	TArray<uint32> AsyncComputePasses;
+	TArray<float, TInlineAllocator<8>> TransientHeapDepths;
+	float HeaderEndDepth{};
+	float HeaderPassBarDepth{};
 };
 
-class FRenderGraphTrack : public FTimingEventsTrack
+class FRenderGraphTrack final : public FBaseTimingTrack
 {
-	INSIGHTS_DECLARE_RTTI(FRenderGraphTrack, FTimingEventsTrack)
-	using Super = FTimingEventsTrack;
+	INSIGHTS_DECLARE_RTTI(FRenderGraphTrack, FBaseTimingTrack)
+	using Super = FBaseTimingTrack;
 
 public:
 	FRenderGraphTrack(const FRenderGraphTimingViewSession& InSharedData);
 
-	//~ Begin FTimingEventsTrack interface
+	//~ Begin FBaseTimingTrack interface
+	void Reset() override;
+	void PreUpdate(const ITimingTrackUpdateContext& Context) override;
 	void Update(const ITimingTrackUpdateContext& Context) override;
+	void PostUpdate(const ITimingTrackUpdateContext& Context) override;
 	void Draw(const ITimingTrackDrawContext& Context) const override;
 	void PostDraw(const ITimingTrackDrawContext& Context) const override;
-	void BuildDrawState(ITimingEventsTrackDrawStateBuilder& Builder, const ITimingTrackUpdateContext& Context) override;
-	void BuildFilteredDrawState(ITimingEventsTrackDrawStateBuilder& Builder, const ITimingTrackUpdateContext& Context) override;
+	void DrawEvent(const ITimingTrackDrawContext& Context, const ITimingEvent& InTimingEvent, EDrawEventMode InDrawMode) const override;
 	const TSharedPtr<const ITimingEvent> GetEvent(float InPosX, float InPosY, const FTimingTrackViewport& Viewport) const override;
 	TSharedPtr<ITimingEventFilter> GetFilterByEvent(const TSharedPtr<const ITimingEvent> InTimingEvent) const;
 	void BuildContextMenu(FMenuBuilder& MenuBuilder) override;
 	void InitTooltip(FTooltipDrawState& Tooltip, const ITimingEvent& HoveredTimingEvent) const override;
-	//~ End FTimingEventsTrack interface
+	//~ End FBaseTimingTrack interface
 
 private:
 	const FRenderGraphTimingViewSession& SharedData;
 
+	void BuildDrawState(FRenderGraphTrackDrawStateBuilder& Builder, const ITimingTrackUpdateContext& Context);
+	void BuildFilteredDrawState(FRenderGraphTrackDrawStateBuilder& Builder, const ITimingTrackUpdateContext& Context);
+
 	FVisibleGraph& AddVisibleGraph(const FTimingTrackViewport& Viewport, const FGraphPacket& Graph, uint32 Color)
 	{
 		const uint32 VisibleIndex = VisibleGraphs.Num();
-		FVisibleGraph& VisibleGraph = VisibleGraphs.Emplace_GetRef(Viewport, Graph, Color, GetHeight());
+		FVisibleGraph& VisibleGraph = VisibleGraphs.Emplace_GetRef(Viewport, Graph, Color, 1.0f);
 		check(VisibleGraph.Packet->VisibleIndex == kInvalidVisibleIndex);
 		VisibleGraph.Packet->VisibleIndex = VisibleIndex;
 		VisibleGraph.Index = VisibleIndex;
+
+		for (uint64 HeapOffset : Graph.TransientHeapOffsets)
+		{
+			VisibleGraph.TransientHeapDepths.Emplace(FMath::CeilToFloat(TransientHeapOffsetToDepth(HeapOffset)));
+		}
+
 		return VisibleGraph;
 	}
 
@@ -229,10 +234,13 @@ private:
 		return nullptr;
 	}
 
-	bool FilterPacket(const FPacket& Packet) const;
-	bool FilterPacketExact(const FPacket& Packet) const;
+	float TransientHeapOffsetToDepth(uint64 HeapOffset) const;
+	float TransientHeapOffsetToDepthAligned(uint64 HeapOffset) const;
+	float TransientHeapOffsetToSlateY(const FTimingViewLayout& Layout, uint64 HeapOffset) const;
 
-	FVector2D MousePosition = FVector2D::ZeroVector;
+	float MouseSlateX{};
+	float MouseSlateY{};
+	float MouseDepthY{};
 
 	TArray<FVisibleGraph> VisibleGraphs;
 
@@ -248,9 +256,6 @@ private:
 
 	TArray<FSpline> Splines;
 
-	const FPacketFilter* PacketFilter{};
-	const FPassPacket* SelectedPass{};
-
 	uint32 GetTextureColor(const FTexturePacket& Texture, uint64 MaxSizeInBytes) const;
 	uint32 GetBufferColor(const FBufferPacket& Buffer, uint64 MaxSizeInBytes) const;
 
@@ -286,16 +291,38 @@ private:
 		Size
 	};
 
+	enum class EVisualizer
+	{
+		Resources,
+		TransientHeaps
+	};
+
 	EResourceShow ResourceShow = EResourceShow::All;
 	EResourceSort ResourceSort = EResourceSort::Creation;
 	EResourceColor ResourceColor = EResourceColor::Type;
+	EVisualizer Visualizer = EVisualizer::Resources;
 	FString FilterText;
 	float FilterSize{};
+	uint64 TransientHeapBytesPerDepthSlot = 1024 * 1024;
 
 	FTooltipDrawState SelectedTooltipState;
+
+	int32 NumLanes;
+	TSharedRef<struct FRenderGraphTrackDrawState> DrawState;
+	TSharedRef<struct FRenderGraphTrackDrawState> FilteredDrawState;
+
+	struct FFilteredDrawStateInfo
+	{
+		double ViewportStartTime = 0.0;
+		double ViewportScaleX = 0.0;
+		double LastBuildDuration = 0.0;
+		TWeakPtr<ITimingEventFilter> LastEventFilter;
+		uint32 Counter = 0;
+		mutable float Opacity = 0.0f;
+	};
+	FFilteredDrawStateInfo FilteredDrawStateInfo;
 };
 
 ENUM_CLASS_FLAGS(FRenderGraphTrack::EResourceShow);
 
-} //namespace RenderGraphInsights
-} //namespace UE
\ No newline at end of file
+}} //namespace UE::RenderGraphInsights
\ No newline at end of file
diff --git a/Engine/Plugins/RenderGraphInsights/Source/RenderGraphInsights/Private/RenderGraphTrackDrawHelper.cpp b/Engine/Plugins/RenderGraphInsights/Source/RenderGraphInsights/Private/RenderGraphTrackDrawHelper.cpp
new file mode 100644
--- /dev/null
+++ b/Engine/Plugins/RenderGraphInsights/Source/RenderGraphInsights/Private/RenderGraphTrackDrawHelper.cpp
@@ -0,0 +1,661 @@
+// Copyright Epic Games, Inc. All Rights Reserved.
+
+#include "RenderGraphTrackDrawHelper.h"
+
+#include "Fonts/FontMeasure.h"
+#include "Brushes/SlateRoundedBoxBrush.h"
+#include "Framework/Application/SlateApplication.h"
+#include "Rendering/DrawElements.h"
+#include "Styling/CoreStyle.h"
+#include "TraceServices/AnalysisService.h"
+
+#include "Insights/Common/PaintUtils.h"
+#include "Insights/Common/TimeUtils.h"
+#include "Insights/ViewModels/TimingEvent.h"
+#include "Insights/ViewModels/TimingTrackViewport.h"
+#include "Insights/ViewModels/BaseTimingTrack.h"
+
+namespace UE { namespace RenderGraphInsights {
+
+constexpr float kBorderColorFactor = 1.25f;
+
+inline FLinearColor ConvertToLinearColor(uint32 InColor)
+{
+	FLinearColor LinearColor;
+	LinearColor.R = ((InColor >> 16) & 0xFF) / 255.0f;
+	LinearColor.G = ((InColor >> 8) & 0xFF) / 255.0f;
+	LinearColor.B = ((InColor) & 0xFF) / 255.0f;
+	LinearColor.A = ((InColor >> 24) & 0xFF) / 255.0f;
+	return LinearColor;
+}
+
+FRenderGraphTrackDrawStateBuilder::FRenderGraphTrackDrawStateBuilder(FRenderGraphTrackDrawState& InState, const FTimingTrackViewport& InViewport)
+	: DrawState(InState)
+	, Viewport(InViewport)
+	, EventFont(FCoreStyle::GetDefaultFontStyle("Regular", 8))
+{
+	DrawState.Reset();
+}
+
+void FRenderGraphTrackDrawStateBuilder::AppendDurationToEventName(FString& InOutEventName, const double InDuration)
+{
+	InOutEventName += TEXT(" (");
+	InOutEventName += TimeUtils::FormatTimeAuto(InDuration);
+	InOutEventName += TEXT(")");
+}
+
+void FRenderGraphTrackDrawStateBuilder::AddEvent(double EventStartTime, double EventEndTime, uint32 EventDepth, const TCHAR* EventName, uint64 EventType, uint32 EventColor)
+{
+	if (EventColor == 0)
+	{
+		EventColor = FTimingEvent::ComputeEventColor(EventName);
+	}
+
+	AddEvent(EventStartTime, EventEndTime, EventDepth, EventColor,
+		[EventName, EventType, EventStartTime, EventEndTime](float Width)
+	{
+		FString Name = EventName;
+
+		if (Width > Name.Len() * 4.0f + 32.0f)
+		{
+			AppendDurationToEventName(Name, EventEndTime - EventStartTime);
+		}
+
+		return Name;
+	});
+}
+
+void FRenderGraphTrackDrawStateBuilder::AddEvent(double EventStartTime, double EventEndTime, uint32 EventDepth, uint32 EventColor, GetEventNameCallback GetEventNameCallback)
+{
+	float EventX1 = Viewport.TimeToSlateUnitsRounded(EventStartTime);
+	if (EventX1 > Viewport.GetWidth())
+	{
+		return;
+	}
+
+	const double RestrictedEndTime = Viewport.RestrictEndTime(EventEndTime);
+	float EventX2 = Viewport.TimeToSlateUnitsRounded(RestrictedEndTime);
+	if (EventX2 < 0)
+	{
+		return;
+	}
+
+	// Timing events are displayed with minimum 1px (including empty ones).
+	if (EventX1 == EventX2)
+	{
+		EventX2 = EventX1 + 1.0f;
+	}
+
+	const int32 Depth = static_cast<int32>(EventDepth);
+	MaxDepthCached = FMath::Max(MaxDepthCached, Depth);
+
+	// Ensure we have enough slots in array. See LastBox[Depth] usage.
+	if (LastBox.Num() <= Depth)
+	{
+		LastBox.AddDefaulted(Depth + 1 - LastBox.Num());
+	}
+
+	// Ensure we have enough slots in array. See LastEventX2[Depth] usage.
+	while (LastEventX2.Num() <= Depth)
+	{
+		LastEventX2.Add(-1.0f);
+	}
+
+	// Limit event width on the viewport's left side.
+	// This also makes the text to be displayed in viewport, for very long events.
+	const float MinX = -2.0f; // -2 allows event border to remain outside screen
+	if (EventX1 < MinX)
+	{
+		EventX1 = MinX;
+	}
+
+	// Limit event width on the viewport's right side.
+	const float MaxX = Viewport.GetWidth() + 2.0f; // +2 allows event border to remain outside screen
+	if (EventX2 > MaxX)
+	{
+		EventX2 = MaxX;
+	}
+
+	const float EventW = EventX2 - EventX1;
+
+	// Optimization...
+	if (EventW == 1.0f && EventX1 == LastEventX2[Depth] - 1.0f)
+	{
+		// Do no draw 1 pixel event if the last event was ended on that pixel.
+		return;
+	}
+
+	const FLinearColor EventColorFill = ConvertToLinearColor(EventColor);
+
+	// Save X2, for current depth.
+	LastEventX2[Depth] = EventX2;
+
+	if (EventW > 2.0f)
+	{
+		FBoxData& Box = LastBox[Depth];
+		if (Box.X1 < Box.X2)
+		{
+			FlushBox(Box, Depth);
+			Box.Reset();
+		}
+
+		// Fill inside of the timing event box.
+		DrawState.InsideBoxes.AddUninitialized();
+		FRenderGraphTrackDrawState::FBoxPrimitive& InsideBox = DrawState.InsideBoxes.Last();
+		InsideBox.DepthY = Depth;
+		InsideBox.DepthH = 1.0f;
+		InsideBox.X = EventX1 + 1.0f;
+		InsideBox.W = EventW - 2.0f;
+		InsideBox.Color = EventColorFill;
+
+		// Add border around the timing event box.
+		DrawState.Borders.AddUninitialized();
+		FRenderGraphTrackDrawState::FBoxPrimitive& BorderBox = DrawState.Borders.Last();
+		BorderBox.DepthY = Depth;
+		BorderBox.DepthH = 1.0f;
+		BorderBox.X = EventX1;
+		BorderBox.W = EventW;
+		BorderBox.Color = FLinearColor(EventColorFill.R * kBorderColorFactor, EventColorFill.G * kBorderColorFactor, EventColorFill.B * kBorderColorFactor, EventColorFill.A);
+	}
+	else // 1px or 2px boxes
+	{
+		FBoxData& Box = LastBox[Depth];
+
+		// Check if we can merge this box with previous one, if any.
+		// Note: We are assuming events are processed in sorted order by X1.
+		if (EventColor == Box.Color && // same color
+			EventX1 <= Box.X2) // overlapping or adjacent
+		{
+			// Merge it with previous box.
+			Box.X2 = EventX2;
+		}
+		else
+		{
+			// Flush previous box, if any.
+			if (Box.X1 < Box.X2)
+			{
+				FlushBox(Box, Depth);
+			}
+
+			// Start new "merge box".
+			Box.X1 = EventX1;
+			Box.X2 = EventX2;
+			Box.Color = EventColor;
+			Box.LinearColor = FLinearColor(EventColorFill.R * kBorderColorFactor, EventColorFill.G * kBorderColorFactor, EventColorFill.B * kBorderColorFactor, EventColorFill.A);
+		}
+	}
+
+	if (EventW > 8.0f)
+	{
+		const FString Name = GetEventNameCallback(EventW);
+
+		const TSharedRef<FSlateFontMeasure> FontMeasureService = FSlateApplication::Get().GetRenderer()->GetFontMeasureService();
+		const int32 LastWholeCharacterIndex = FontMeasureService->FindLastWholeCharacterIndexBeforeOffset(Name, EventFont, FMath::RoundToInt(EventW - 2.0f));
+
+		if (LastWholeCharacterIndex >= 0)
+		{
+			const bool bIsDarkColor = (EventColorFill.GetLuminance() < 0.4f);
+
+			DrawState.Texts.AddDefaulted();
+			FRenderGraphTrackDrawState::FTextPrimitive& DrawText = DrawState.Texts.Last();
+			DrawText.Depth = Depth;
+			DrawText.X = EventX1 + 2.0f;
+			DrawText.Text = Name.Left(LastWholeCharacterIndex + 1);
+			DrawText.bWhite = bIsDarkColor;
+			DrawText.Color = FLinearColor(EventColorFill.R * kBorderColorFactor, EventColorFill.G * kBorderColorFactor, EventColorFill.B * kBorderColorFactor, EventColorFill.A);
+		}
+	}
+}
+
+void FRenderGraphTrackDrawStateBuilder::AddEvent(double EventStartTime, double EventEndTime, float EventDepthY, float EventDepthH, uint32 EventColor, GetEventNameCallback GetEventNameCallback)
+{
+	float EventX1 = Viewport.TimeToSlateUnitsRounded(EventStartTime);
+	if (EventX1 > Viewport.GetWidth())
+	{
+		return;
+	}
+
+	const double RestrictedEndTime = Viewport.RestrictEndTime(EventEndTime);
+	float EventX2 = Viewport.TimeToSlateUnitsRounded(RestrictedEndTime);
+	if (EventX2 < 0)
+	{
+		return;
+	}
+
+	// Timing events are displayed with minimum 1px (including empty ones).
+	if (EventX1 == EventX2)
+	{
+		EventX2 = EventX1 + 1.0f;
+	}
+
+	MaxDepth = FMath::Max<int32>(MaxDepth, FMath::CeilToFloat(EventDepthY + EventDepthH));
+
+	// Limit event width on the viewport's left side.
+	// This also makes the text to be displayed in viewport, for very long events.
+	const float MinX = -2.0f; // -2 allows event border to remain outside screen
+	if (EventX1 < MinX)
+	{
+		EventX1 = MinX;
+	}
+
+	// Limit event width on the viewport's right side.
+	const float MaxX = Viewport.GetWidth() + 2.0f; // +2 allows event border to remain outside screen
+	if (EventX2 > MaxX)
+	{
+		EventX2 = MaxX;
+	}
+
+	const float EventW = EventX2 - EventX1;
+	const FLinearColor EventColorFill = ConvertToLinearColor(EventColor);
+
+	if (EventW > 0.0f)
+	{
+		DrawState.Boxes.AddUninitialized();
+		FRenderGraphTrackDrawState::FBoxPrimitive& DrawBox = DrawState.Boxes.Last();
+		DrawBox.X = EventX1;
+		DrawBox.W = EventW;
+		DrawBox.DepthY = EventDepthY;
+		DrawBox.DepthH = EventDepthH;
+		DrawBox.Color = FLinearColor(EventColorFill.R * kBorderColorFactor, EventColorFill.G * kBorderColorFactor, EventColorFill.B * kBorderColorFactor, EventColorFill.A);
+
+		// Fill inside of the timing event box.
+		DrawState.InsideBoxes.AddUninitialized();
+		FRenderGraphTrackDrawState::FBoxPrimitive& InsideBox = DrawState.InsideBoxes.Last();
+		InsideBox.X = EventX1 + 1.0f;
+		InsideBox.W = EventW - 2.0f;
+		InsideBox.DepthY = EventDepthY;
+		InsideBox.DepthH = EventDepthH;
+		InsideBox.Color = EventColorFill;
+
+		// Add border around the timing event box.
+		DrawState.Borders.AddUninitialized();
+		FRenderGraphTrackDrawState::FBoxPrimitive& BorderBox = DrawState.Borders.Last();
+		BorderBox.X = EventX1;
+		BorderBox.W = EventW;
+		BorderBox.DepthY = EventDepthY;
+		BorderBox.DepthH = EventDepthH;
+		BorderBox.Color = FLinearColor(EventColorFill.R * kBorderColorFactor, EventColorFill.G * kBorderColorFactor, EventColorFill.B * kBorderColorFactor, EventColorFill.A);
+	}
+
+	// Draw the name of the timing event.
+	if (EventW > 8.0f && EventDepthH >= 1.0f)
+	{
+		const FString Name = GetEventNameCallback(EventW);
+
+		const TSharedRef<FSlateFontMeasure> FontMeasureService = FSlateApplication::Get().GetRenderer()->GetFontMeasureService();
+		const int32 LastWholeCharacterIndex = FontMeasureService->FindLastWholeCharacterIndexBeforeOffset(Name, EventFont, FMath::RoundToInt(EventW - 2.0f));
+
+		if (LastWholeCharacterIndex >= 0)
+		{
+			// Grey threshold is shifted toward black (0.4 instead of 0.5 in test below) due to "area rule":
+			// a large gray surface (background of a timing event in this case) is perceived lighter than a smaller area (text pixels).
+			// Ref: https://books.google.ro/books?id=0pVr7dhmdWYC
+			const bool bIsDarkColor = (EventColorFill.GetLuminance() < 0.4f);
+
+			DrawState.Texts.AddDefaulted();
+			FRenderGraphTrackDrawState::FTextPrimitive& DrawText = DrawState.Texts.Last();
+			DrawText.X = EventX1 + 2.0f;
+			DrawText.Depth = EventDepthY;
+			DrawText.Text = Name.Left(LastWholeCharacterIndex + 1);
+			DrawText.bWhite = bIsDarkColor;
+			DrawText.Color = FLinearColor(EventColorFill.R * kBorderColorFactor, EventColorFill.G * kBorderColorFactor, EventColorFill.B * kBorderColorFactor, EventColorFill.A);
+		}
+	}
+}
+
+void FRenderGraphTrackDrawStateBuilder::FlushBox(const FBoxData& Box, const int32 Depth)
+{
+	DrawState.Boxes.AddUninitialized();
+	FRenderGraphTrackDrawState::FBoxPrimitive& DrawBox = DrawState.Boxes.Last();
+	DrawBox.DepthY = Depth;
+	DrawBox.DepthH = 1.0f;
+	DrawBox.X = Box.X1;
+	DrawBox.W = Box.X2 - Box.X1;
+	DrawBox.Color = Box.LinearColor;
+}
+
+void FRenderGraphTrackDrawStateBuilder::Flush()
+{
+	for (int32 Depth = 0; Depth <= MaxDepthCached; ++Depth)
+	{
+		const FBoxData& Box = LastBox[Depth];
+		if (Box.X1 < Box.X2)
+		{
+			FlushBox(Box, Depth);
+		}
+	}
+}
+
+FRenderGraphTrackDrawHelper::FRenderGraphTrackDrawHelper(const ITimingTrackDrawContext& Context)
+	: DrawContext(Context.GetDrawContext())
+	, Viewport(Context.GetViewport())
+	, ParentHelper(Context.GetHelper())
+	, WhiteBrush(ParentHelper.GetWhiteBrush())
+	, EventBorderBrush(ParentHelper.GetEventBorderBrush())
+	, HoveredEventBorderBrush(ParentHelper.GetHoveredEventBorderBrush())
+	, SelectedEventBorderBrush(ParentHelper.GetSelectedEventBorderBrush())
+	, ValidAreaColor(ParentHelper.GetValidAreaColor())
+	, InvalidAreaColor(ParentHelper.GetInvalidAreaColor())
+	, EdgeColor(ParentHelper.GetEdgeColor())
+	, EventFont(ParentHelper.GetEventFont())
+	, ReservedLayerId(ParentHelper.GetFirstLayerId())
+{}
+
+void FRenderGraphTrackDrawHelper::DrawEvents(const FRenderGraphTrackDrawState& DrawState, const FBaseTimingTrack& Track) const
+{
+	const float TrackH = Track.GetHeight();
+	if (TrackH > 0.0f)
+	{
+		const FTimingViewLayout& Layout = Viewport.GetLayout();
+
+		float TopLaneY = Track.GetPosY() + Layout.TimelineDY + 1.0f;
+
+		// Draw filled boxes (merged borders).
+		{
+			const int32 EventFillLayerId = ReservedLayerId + int32(EDrawLayer::EventFill);
+
+			for (const FRenderGraphTrackDrawState::FBoxPrimitive& Box : DrawState.Boxes)
+			{
+				const float Y = GetLaneY(TopLaneY, Layout, Box.DepthY);
+				const float H = GetLaneHeight(Layout, Box.DepthH);
+				DrawContext.DrawBox(EventFillLayerId, Box.X, Y, Box.W, H, WhiteBrush, Box.Color);
+			}
+		}
+
+		// Draw filled boxes (event inside area).
+		if (Layout.EventH > 2.0f)
+		{
+			const int32 EventFillLayerId = ReservedLayerId + int32(EDrawLayer::EventFill);
+
+			for (const FRenderGraphTrackDrawState::FBoxPrimitive& Box : DrawState.InsideBoxes)
+			{
+				const float Y = GetLaneY(TopLaneY, Layout, Box.DepthY) + 1.0f;
+				const float H = GetLaneHeight(Layout, Box.DepthH) - 2.0f;
+				DrawContext.DrawBox(EventFillLayerId, Box.X, Y, Box.W, H, WhiteBrush, Box.Color);
+			}
+		}
+
+		// Draw borders.
+		{
+			const int32 EventBorderLayerId = ReservedLayerId + int32(EDrawLayer::EventBorder);
+
+			for (const FRenderGraphTrackDrawState::FBoxPrimitive& Box : DrawState.Borders)
+			{
+				const float Y = GetLaneY(TopLaneY, Layout, Box.DepthY);
+				const float H = GetLaneHeight(Layout, Box.DepthH);
+				DrawContext.DrawBox(EventBorderLayerId, Box.X, Y, Box.W, H, WhiteBrush, Box.Color);
+			}
+		}
+
+		// Draw texts.
+		if (Layout.EventH > 10.0f)
+		{
+			const int32 EventTextLayerId = ReservedLayerId + int32(EDrawLayer::EventText);
+			for (const FRenderGraphTrackDrawState::FTextPrimitive& Text : DrawState.Texts)
+			{
+				const float Y = GetLaneY(TopLaneY, Layout, Text.Depth) + 1.0f;
+				DrawContext.DrawText(EventTextLayerId, Text.X, Y, Text.Text, EventFont, Text.bWhite ? FLinearColor::White : FLinearColor::Black);
+			}
+		}
+
+		{
+			const EDrawLayer SplineLayer = EDrawLayer::Foreground;
+			for (const FSplinePrimitive& Spline : DrawState.Splines)
+			{
+				DrawSpline(Track, Spline, SplineLayer);
+			}
+		}
+	}
+}
+
+void FRenderGraphTrackDrawHelper::DrawFadedEvents(const FRenderGraphTrackDrawState& DrawState, const FBaseTimingTrack& Track, const float Opacity) const
+{
+	const float TrackH = Track.GetHeight();
+	if (TrackH > 0.0f)
+	{
+		const FTimingViewLayout& Layout = Viewport.GetLayout();
+
+		float TopLaneY = Track.GetPosY() + Layout.TimelineDY + 1.0f;
+
+		// Draw filled boxes (merged borders).
+		{
+			const int32 EventFillLayerId = ReservedLayerId + int32(EDrawLayer::EventFill);
+
+			for (const FRenderGraphTrackDrawState::FBoxPrimitive& Box : DrawState.Boxes)
+			{
+				const float Y = GetLaneY(TopLaneY, Layout, Box.DepthY);
+				const float H = GetLaneHeight(Layout, Box.DepthH);
+				DrawContext.DrawBox(EventFillLayerId, Box.X, Y, Box.W, H, WhiteBrush, Box.Color.CopyWithNewOpacity(Opacity));
+			}
+		}
+
+		// Draw filled boxes (event inside area).
+		if (Layout.EventH > 2.0f)
+		{
+			const int32 EventFillLayerId = ReservedLayerId + int32(EDrawLayer::EventFill);
+
+			for (const FRenderGraphTrackDrawState::FBoxPrimitive& Box : DrawState.InsideBoxes)
+			{
+				const float Y = GetLaneY(TopLaneY, Layout, Box.DepthY) + 1.0f;
+				const float H = GetLaneHeight(Layout, Box.DepthH) - 2.0f;
+				DrawContext.DrawBox(EventFillLayerId, Box.X, Y, Box.W, H, WhiteBrush, Box.Color.CopyWithNewOpacity(Opacity));
+			}
+		}
+
+		// Draw borders.
+		{
+			const int32 EventBorderLayerId = ReservedLayerId + int32(EDrawLayer::EventBorder);
+
+			for (const FRenderGraphTrackDrawState::FBoxPrimitive& Box : DrawState.Borders)
+			{
+				const float Y = GetLaneY(TopLaneY, Layout, Box.DepthY);
+				const float H = GetLaneHeight(Layout, Box.DepthH);
+				FSlateRoundedBoxBrush Brush(FLinearColor(0.0f, 0.0f, 0.0f, 0.0f), 0.0f, Box.Color.CopyWithNewOpacity(Opacity), 1.0f);
+				DrawContext.DrawBox(EventBorderLayerId, Box.X, Y, Box.W, H, &Brush, FLinearColor(0.0f, 0.0f, 0.0f, 0.0f));
+			}
+		}
+
+		// Draw texts.
+		if (Layout.EventH > 10.0f)
+		{
+			const FLinearColor WhiteColor(1.0f, 1.0f, 1.0f, Opacity);
+			const FLinearColor BlackColor(0.0f, 0.0f, 0.0f, Opacity);
+
+			const int32 EventTextLayerId = ReservedLayerId + int32(EDrawLayer::EventText);
+			for (const FRenderGraphTrackDrawState::FTextPrimitive& Text : DrawState.Texts)
+			{
+				const float Y = GetLaneY(TopLaneY, Layout, Text.Depth) + 1.0f;
+				DrawContext.DrawText(EventTextLayerId, Text.X, Y, Text.Text, EventFont, Text.bWhite ? WhiteColor : BlackColor);
+			}
+		}
+	}
+}
+
+void FRenderGraphTrackDrawHelper::DrawTimingEventHighlight(const FBaseTimingTrack& Track, double StartTime, double EndTime, float DepthY, float DepthH, EDrawEventMode Mode) const
+{
+	float EventX1 = Viewport.TimeToSlateUnitsRounded(StartTime);
+	if (EventX1 > Viewport.GetWidth())
+	{
+		return;
+	}
+
+	EndTime = Viewport.RestrictEndTime(EndTime);
+	float EventX2 = Viewport.TimeToSlateUnitsRounded(EndTime);
+	if (EventX2 < 0)
+	{
+		return;
+	}
+
+	// Timing events are displayed with minimum 1px (including empty ones).
+	if (EventX1 == EventX2)
+	{
+		EventX2 = EventX1 + 1.0f;
+	}
+
+	// Limit event width on the viewport's left side.
+	const float MinX = -2.0f; // -2 allows event border to remain outside screen
+	if (EventX1 < MinX)
+	{
+		EventX1 = MinX;
+	}
+
+	// Limit event width on the viewport's right side.
+	const float MaxX = Viewport.GetWidth() + 2.0f; // +2 allows event border to remain outside screen
+	if (EventX2 > MaxX)
+	{
+		EventX2 = MaxX;
+	}
+
+	const FTimingViewLayout& Layout = Viewport.GetLayout();
+
+	const float EventW = EventX2 - EventX1;
+	const float EventY = Track.GetPosY() + GetLaneY(Layout, DepthY);
+	const float EventH = GetLaneHeight(Layout, DepthH);
+
+	const int32 LayerId = ReservedLayerId + int32(EDrawLayer::EventHighlight);
+
+	constexpr float BorderOffset = 2.0f;
+
+	if (Mode == EDrawEventMode::Hovered)
+	{
+		const FLinearColor Color(1.0f, 1.0f, 0.0f, 1.0f); // yellow
+
+		// Draw border around the timing event box.
+		FSlateRoundedBoxBrush Brush(FLinearColor(0.0f, 0.0f, 0.0f, 0.0f), 2.0f, Color, 2.0f);
+		DrawContext.DrawBox(LayerId, EventX1 - BorderOffset, EventY - BorderOffset, EventW + 2 * BorderOffset, EventH + 2 * BorderOffset, &Brush, FLinearColor(0.0f, 0.0f, 0.0f, 0.0f));
+	}
+	else // EDrawEventMode::Selected or EDrawEventMode::SelectedAndHovered
+	{
+		// Animate color from white (if selected and hovered) or yellow (if only selected) to black, using a squared sine function.
+		const double Time = static_cast<double>(FPlatformTime::Cycles64()) * FPlatformTime::GetSecondsPerCycle64();
+		float S = FMath::Sin(2.0 * Time);
+		S = S * S; // squared, to ensure only positive [0 - 1] values
+		const float Blue = (Mode == EDrawEventMode::SelectedAndHovered) ? 0.0f : S;
+		const FLinearColor Color(S, S, Blue, 1.0f);
+
+		// Draw border around the timing event box.
+		FSlateRoundedBoxBrush Brush(FLinearColor(0.0f, 0.0f, 0.0f, 0.0f), 2.0f, Color, 2.0f);
+		DrawContext.DrawBox(LayerId, EventX1 - BorderOffset, EventY - BorderOffset, EventW + 2 * BorderOffset, EventH + 2 * BorderOffset, &Brush, FLinearColor(0.0f, 0.0f, 0.0f, 0.0f));
+	}
+}
+
+void FRenderGraphTrackDrawHelper::DrawTrackHeader(const FBaseTimingTrack& Track) const
+{
+	const int32 HeaderLayerId = ReservedLayerId + int32(EDrawLayer::HeaderBackground);
+	const int32 HeaderTextLayerId = ReservedLayerId + int32(EDrawLayer::HeaderText);
+
+	DrawTrackHeader(Track, HeaderLayerId, HeaderTextLayerId);
+}
+
+void FRenderGraphTrackDrawHelper::DrawTrackHeader(const FBaseTimingTrack& Track, const int32 HeaderLayerId, const int32 HeaderTextLayerId) const
+{
+	const float TrackY = Track.GetPosY();
+	const float TrackH = Track.GetHeight();
+
+	if (TrackH > 0.0f)
+	{
+		// Draw a horizontal line between timelines (top line of a track).
+		DrawContext.DrawBox(HeaderLayerId, 0.0f, TrackY, Viewport.GetWidth(), 1.0f, WhiteBrush, EdgeColor);
+
+		// Draw header with name of timeline.
+		if (TrackH > 4.0f)
+		{
+			const TSharedRef<FSlateFontMeasure> FontMeasureService = FSlateApplication::Get().GetRenderer()->GetFontMeasureService();
+			float TextWidth = FontMeasureService->Measure(Track.GetName(), EventFont).X;
+
+			constexpr float PinWidth = 8.0f;
+			if (Track.IsSelected())
+			{
+				TextWidth += PinWidth;
+			}
+
+			const float HeaderX = 0.0f;
+			const float HeaderW = TextWidth + 4.0f;
+
+			const float HeaderY = TrackY + 1.0f;
+			const float HeaderH = FMath::Min(12.0f, Track.GetHeight() - 1.0f);
+
+			if (HeaderH > 0)
+			{
+				DrawContext.DrawBox(HeaderLayerId, HeaderX, HeaderY, HeaderW, HeaderH, WhiteBrush, EdgeColor);
+
+				const FLinearColor TextColor = GetTrackNameTextColor(Track);
+
+				float TextX = HeaderX + 2.0f;
+				const float TextY = HeaderY + HeaderH / 2.0f - 7.0f;
+
+				if (Track.IsSelected())
+				{
+					DrawContext.DrawText(HeaderTextLayerId, TextX, TextY, TEXT("\u25CA"), EventFont, TextColor); // lozenge shape
+					TextX += PinWidth;
+				}
+
+				DrawContext.DrawText(HeaderTextLayerId, TextX, TextY, Track.GetName(), EventFont, TextColor);
+			}
+		}
+	}
+}
+
+void FRenderGraphTrackDrawHelper::DrawSpline(const FBaseTimingTrack& Track, const FSplinePrimitive& Spline, EDrawLayer Layer) const
+{
+	const float ViewportWidth = Viewport.GetWidth();
+
+	const float Guardband = 1024.0f;
+	float StartX = FMath::Max(Spline.Start.X               , -Guardband);
+	float EndX   = FMath::Min(Spline.Start.X + Spline.End.X,  Guardband + ViewportWidth);
+	const float Width = EndX - StartX;
+
+	if (StartX > ViewportWidth || EndX < 0.0f)
+	{
+		return;
+	}
+
+	const FTimingViewLayout& Layout = Viewport.GetLayout();
+
+	const float TopLaneY  = Track.GetPosY() + Layout.TimelineDY + 1.0f;
+	const float StartY    = GetLaneY(TopLaneY, Layout, Spline.Start.Y);
+	const float EndY      = GetChildLaneY(Layout, Spline.End.Y);
+
+	DrawContext.DrawSpline(
+		ReservedLayerId + int32(Layer),
+		StartX,
+		StartY,
+		FVector2D::ZeroVector,
+		Spline.StartDir,
+		FVector2D(Width, EndY),
+		Spline.EndDir,
+		Spline.Thickness,
+		Spline.Tint
+	);
+}
+
+void FRenderGraphTrackDrawHelper::DrawBox(const FBaseTimingTrack& Track, float SlateX, float DepthY, float SlateW, float DepthH, FLinearColor Color, EDrawLayer Layer) const
+{
+	const float ViewportWidth = Viewport.GetWidth();
+
+	const float Guardband = 1024.0f;
+	const float SlateMinX = FMath::Max(SlateX         , -Guardband);
+	const float SlateMaxX = FMath::Min(SlateX + SlateW,  Guardband + ViewportWidth);
+	SlateW = SlateMaxX - SlateMinX;
+
+	if (SlateMinX > ViewportWidth || SlateMaxX < 0.0f)
+	{
+		return;
+	}
+
+	const FTimingViewLayout& Layout = Viewport.GetLayout();
+
+	const float SlateY = GetLaneY(Layout, DepthY) + Track.GetPosY();
+	const float SlateH = GetLaneHeight(Layout, DepthH);
+
+	DrawContext.DrawBox(ReservedLayerId + int32(Layer), SlateMinX, SlateY, SlateW, SlateH, WhiteBrush, Color);
+}
+
+FLinearColor FRenderGraphTrackDrawHelper::GetTrackNameTextColor(const FBaseTimingTrack& Track) const
+{
+	return  Track.IsHovered() ? FLinearColor(1.0f, 1.0f, 0.0f, 1.0f) :
+		Track.IsSelected() ? FLinearColor(1.0f, 1.0f, 0.5f, 1.0f) :
+		FLinearColor(1.0f, 1.0f, 1.0f, 1.0f);
+}
+
+}} // namespace UE::RenderGraphInsights
\ No newline at end of file
diff --git a/Engine/Plugins/RenderGraphInsights/Source/RenderGraphInsights/Private/RenderGraphTrackDrawHelper.h b/Engine/Plugins/RenderGraphInsights/Source/RenderGraphInsights/Private/RenderGraphTrackDrawHelper.h
new file mode 100644
--- /dev/null
+++ b/Engine/Plugins/RenderGraphInsights/Source/RenderGraphInsights/Private/RenderGraphTrackDrawHelper.h
@@ -0,0 +1,211 @@
+// Copyright Epic Games, Inc. All Rights Reserved.
+
+#pragma once
+
+#include "CoreMinimal.h"
+#include "Fonts/SlateFontInfo.h"
+#include "Math/Color.h"
+
+#include "Insights/ViewModels/ITimingEvent.h"
+#include "Insights/ViewModels/ITimingViewDrawHelper.h"
+#include "Insights/ViewModels/TimingViewLayout.h"
+
+struct FSlateBrush;
+struct FDrawContext;
+
+class ITimingTrackDrawContext;
+class FTimingTrackViewport;
+class FBaseTimingTrack;
+
+enum class EDrawEventMode : uint32;
+
+namespace UE { namespace RenderGraphInsights {
+	
+inline float GetChildLaneY(const FTimingViewLayout& Layout, float DepthY)
+{
+	return (Layout.EventH + Layout.EventDY) * DepthY;
+}
+
+inline float GetLaneY(float TopLaneY, const FTimingViewLayout& Layout, float DepthY)
+{
+	return TopLaneY + GetChildLaneY(Layout, DepthY);
+}
+
+inline float GetLaneY(const FTimingViewLayout& Layout, float DepthY)
+{
+	return GetLaneY(Layout.TimelineDY + 1.0f, Layout, DepthY);
+}
+
+inline float GetLaneHeight(const FTimingViewLayout& Layout, float DepthH)
+{
+	return Layout.EventH * DepthH + Layout.EventDY * (FMath::CeilToFloat(DepthH) - 1);
+}
+
+inline float GetDepthFromLaneY(const FTimingViewLayout& Layout, float Y)
+{
+	return (Y - 1.0f - Layout.TimelineDY) / (Layout.EventH + Layout.EventDY);
+}
+
+class FRenderGraphBaseTrack;
+
+enum class EDrawLayer : int32
+{
+	Background,
+	EventBorder,
+	EventFill,
+	EventText,
+	EventHighlight,
+	Relation,
+	HeaderBackground,
+	HeaderText,
+	Foreground,
+	Count,
+};
+
+struct FSplinePrimitive
+{
+	float Thickness = 1.0f;
+	FVector2D Start = FVector2D::ZeroVector;
+	FVector2D StartDir = FVector2D::ZeroVector;
+	FVector2D End = FVector2D::ZeroVector;
+	FVector2D EndDir = FVector2D::ZeroVector;
+	FLinearColor Tint = FLinearColor::White;
+};
+
+struct FRenderGraphTrackDrawState
+{
+	struct FBoxPrimitive
+	{
+		float DepthY{};
+		float DepthH{};
+		float X{};
+		float W{};
+		FLinearColor Color = FLinearColor::White;
+	};
+
+	struct FTextPrimitive
+	{
+		float Depth{};
+		float X{};
+		FString Text{};
+		bool bWhite{};
+		FLinearColor Color = FLinearColor::White;
+	};
+
+	void Reset()
+	{
+		Boxes.Reset();
+		InsideBoxes.Reset();
+		Borders.Reset();
+		Texts.Reset();
+		Splines.Reset();
+	}
+
+	TArray<FBoxPrimitive> Boxes;
+	TArray<FBoxPrimitive> InsideBoxes;
+	TArray<FBoxPrimitive> Borders;
+	TArray<FTextPrimitive> Texts;
+	TArray<FSplinePrimitive> Splines;
+};
+
+class FRenderGraphTrackDrawStateBuilder
+{
+private:
+	struct FBoxData
+	{
+		float X1{};
+		float X2{};
+		uint32 Color{};
+		FLinearColor LinearColor;
+
+		void Reset() { X1 = 0.0f; X2 = 0.0f; Color = 0; }
+	};
+
+public:
+	explicit FRenderGraphTrackDrawStateBuilder(FRenderGraphTrackDrawState& InState, const FTimingTrackViewport& InViewport);
+
+	FRenderGraphTrackDrawStateBuilder(const FRenderGraphTrackDrawStateBuilder&) = delete;
+	FRenderGraphTrackDrawStateBuilder& operator=(const FRenderGraphTrackDrawStateBuilder&) = delete;
+
+	const FTimingTrackViewport& GetViewport() const { return Viewport; }
+
+	using GetEventNameCallback = TFunctionRef<const FString(float /*AvailableWidth*/)>;
+
+	void AddEvent(double InEventStartTime, double InEventEndTime, uint32 InEventDepth, const TCHAR* InEventName, uint64 InEventType = 0, uint32 InEventColor = 0);
+	void AddEvent(double InEventStartTime, double InEventEndTime, uint32 InEventDepth, uint32 InEventColor, GetEventNameCallback InGetEventNameCallback);
+	void AddEvent(double InEventStartTime, double InEventEndTime, float InEventDepthY, float InEventDepthH, uint32 InEventColor, GetEventNameCallback InGetEventNameCallback);
+	void Flush();
+
+	void AddSpline(const FSplinePrimitive& Spline)
+	{
+		DrawState.Splines.Add(Spline);
+	}
+
+	int32 GetMaxDepth() const { return FMath::Max(MaxDepth, MaxDepthCached); }
+
+private:
+	void FlushBox(const FBoxData& Box, const int32 Depth);
+
+	static void AppendDurationToEventName(FString& InOutEventName, const double InDuration);
+
+	FRenderGraphTrackDrawState& DrawState; // cached draw state to build
+	const FTimingTrackViewport& Viewport;
+
+	int32 MaxDepth = -1;
+	int32 MaxDepthCached = -1;
+
+	TArray<float> LastEventX2; // X2 value for last event on each depth
+	TArray<FBoxData> LastBox;
+
+	const FSlateFontInfo EventFont;
+};
+
+class FRenderGraphTrackDrawHelper final
+{
+public:
+	explicit FRenderGraphTrackDrawHelper(const ITimingTrackDrawContext& Context);
+
+	FRenderGraphTrackDrawHelper(const FRenderGraphTrackDrawHelper&) = delete;
+	FRenderGraphTrackDrawHelper& operator=(const FRenderGraphTrackDrawHelper&) = delete;
+
+	const FDrawContext& GetDrawContext() const { return DrawContext; }
+	const FTimingTrackViewport& GetViewport() const { return Viewport; }
+
+	void DrawEvents(const FRenderGraphTrackDrawState& DrawState, const FBaseTimingTrack& Track) const;
+	void DrawFadedEvents(const FRenderGraphTrackDrawState& DrawState, const FBaseTimingTrack& Track, const float Opacity = 0.1f) const;
+
+	void DrawTimingEventHighlight(const FBaseTimingTrack& Track, double StartTime, double EndTime, float DepthY, float DepthH, EDrawEventMode Mode) const;
+
+	void DrawTrackHeader(const FBaseTimingTrack& Track) const;
+	void DrawTrackHeader(const FBaseTimingTrack& Track, const int32 HeaderLayerId, const int32 HeaderTextLayerId) const;
+
+	void DrawSpline(const FBaseTimingTrack& Track, const FSplinePrimitive& Spline, EDrawLayer Layer) const;
+
+	void DrawBox(const FBaseTimingTrack& Track, float SlateX, float DepthY, float SlateW, float DepthH, FLinearColor Color, EDrawLayer Layer) const;
+
+	FLinearColor GetTrackNameTextColor(const FBaseTimingTrack& Track) const;
+
+	FLinearColor GetEdgeColor() const { return EdgeColor; }
+	const FSlateFontInfo& GetEventFont() const { return EventFont; }
+	const FSlateBrush* GetWhiteBrush() const { return WhiteBrush; }
+	int32 GetFirstLayerId() const { return ReservedLayerId; }
+	int32 GetNumLayerIds() const { return int32(EDrawLayer::Count); }
+
+private:
+
+	const FDrawContext& DrawContext;
+	const FTimingTrackViewport& Viewport;
+	const ITimingViewDrawHelper& ParentHelper;
+
+	const FSlateBrush* WhiteBrush;
+	const FSlateBrush* EventBorderBrush;
+	const FSlateBrush* HoveredEventBorderBrush;
+	const FSlateBrush* SelectedEventBorderBrush;
+	const FLinearColor ValidAreaColor;
+	const FLinearColor InvalidAreaColor;
+	const FLinearColor EdgeColor;
+	const FSlateFontInfo EventFont;
+	const int32 ReservedLayerId{};
+};
+
+}} // UE::RenderGraphInsights
\ No newline at end of file
diff --git a/Engine/Source/Developer/TraceInsights/Private/Insights/ViewModels/TimingViewDrawHelper.h b/Engine/Source/Developer/TraceInsights/Private/Insights/ViewModels/TimingViewDrawHelper.h
--- a/Engine/Source/Developer/TraceInsights/Private/Insights/ViewModels/TimingViewDrawHelper.h
+++ b/Engine/Source/Developer/TraceInsights/Private/Insights/ViewModels/TimingViewDrawHelper.h
@@ -169,8 +169,13 @@ public:
 
 	// ITimingViewDrawHelper interface
 	virtual const FSlateBrush* GetWhiteBrush() const override { return WhiteBrush; }
+	virtual const FSlateBrush* GetEventBorderBrush() const override { return EventBorderBrush; }
+	virtual const FSlateBrush* GetHoveredEventBorderBrush() const override { return HoveredEventBorderBrush; }
+	virtual const FSlateBrush* GetSelectedEventBorderBrush() const override { return SelectedEventBorderBrush; }
 	virtual const FSlateFontInfo& GetEventFont() const override { return EventFont; }
 	virtual FLinearColor GetEdgeColor() const override { return EdgeColor; }
+	virtual FLinearColor GetValidAreaColor() const override { return ValidAreaColor; }
+	virtual FLinearColor GetInvalidAreaColor() const override { return InvalidAreaColor; }
 	virtual FLinearColor GetTrackNameTextColor(const FBaseTimingTrack& Track) const override;
 	virtual int32 GetHeaderBackgroundLayerId() const override { return ReservedLayerId + ToInt32(EDrawLayer::HeaderBackground); }
 	virtual int32 GetHeaderTextLayerId() const override { return ReservedLayerId + ToInt32(EDrawLayer::HeaderText); }
diff --git a/Engine/Source/Developer/TraceInsights/Private/Insights/Widgets/STimingView.h b/Engine/Source/Developer/TraceInsights/Private/Insights/Widgets/STimingView.h
--- a/Engine/Source/Developer/TraceInsights/Private/Insights/Widgets/STimingView.h
+++ b/Engine/Source/Developer/TraceInsights/Private/Insights/Widgets/STimingView.h
@@ -269,6 +269,17 @@ public:
 	virtual Insights::FSelectedTrackChangedDelegate& OnSelectedTrackChanged() override { return OnSelectedTrackChangedDelegate; }
 	virtual Insights::FSelectedEventChangedDelegate& OnSelectedEventChanged() override { return OnSelectedEventChangedDelegate; }
 
+	virtual void ResetSelectedEvent() override
+	{
+		if (SelectedEvent)
+		{
+			SelectedEvent.Reset();
+			OnSelectedEventChanged();
+		}
+	}
+
+	virtual void ResetEventFilter() override { SetEventFilter(nullptr); }
+
 	virtual void PreventThrottling() override;
 	virtual void AddOverlayWidget(const TSharedRef<SWidget>& InWidget) override;
 
diff --git a/Engine/Source/Developer/TraceInsights/Public/Insights/ITimingViewSession.h b/Engine/Source/Developer/TraceInsights/Public/Insights/ITimingViewSession.h
--- a/Engine/Source/Developer/TraceInsights/Public/Insights/ITimingViewSession.h
+++ b/Engine/Source/Developer/TraceInsights/Public/Insights/ITimingViewSession.h
@@ -110,6 +110,14 @@ public:
 
 	//////////////////////////////////////////////////
 
+	/** Resets the selected event back to empty. */
+	virtual void ResetSelectedEvent() = 0;
+
+	/** Resets the event filter back to empty. */
+	virtual void ResetEventFilter() = 0;
+
+	//////////////////////////////////////////////////
+
 	/** Prevents mouse movements from throttling application updates. */
 	virtual void PreventThrottling() = 0;
 
diff --git a/Engine/Source/Developer/TraceInsights/Public/Insights/ViewModels/ITimingViewDrawHelper.h b/Engine/Source/Developer/TraceInsights/Public/Insights/ViewModels/ITimingViewDrawHelper.h
--- a/Engine/Source/Developer/TraceInsights/Public/Insights/ViewModels/ITimingViewDrawHelper.h
+++ b/Engine/Source/Developer/TraceInsights/Public/Insights/ViewModels/ITimingViewDrawHelper.h
@@ -13,12 +13,17 @@ class ITimingViewDrawHelper
 {
 public:
 	virtual const FSlateBrush* GetWhiteBrush() const = 0;
+	virtual const FSlateBrush* GetEventBorderBrush() const = 0;
+	virtual const FSlateBrush* GetHoveredEventBorderBrush() const = 0;
+	virtual const FSlateBrush* GetSelectedEventBorderBrush() const = 0;
 	virtual const FSlateFontInfo& GetEventFont() const = 0;
 	virtual FLinearColor GetEdgeColor() const = 0;
+	virtual FLinearColor GetValidAreaColor() const = 0;
+	virtual FLinearColor GetInvalidAreaColor() const = 0;
 	virtual FLinearColor GetTrackNameTextColor(const FBaseTimingTrack& Track) const = 0;
 	virtual int32 GetHeaderBackgroundLayerId() const = 0;
 	virtual int32 GetHeaderTextLayerId() const = 0;
 	virtual int32 GetRelationLayerId() const = 0;
 	virtual int32 GetFirstLayerId() const = 0;
 	virtual int32 GetNumLayerIds() const = 0;
-};
+};
\ No newline at end of file
diff --git a/Engine/Source/Runtime/D3D12RHI/Private/D3D12TransientResourceAllocator.cpp b/Engine/Source/Runtime/D3D12RHI/Private/D3D12TransientResourceAllocator.cpp
--- a/Engine/Source/Runtime/D3D12RHI/Private/D3D12TransientResourceAllocator.cpp
+++ b/Engine/Source/Runtime/D3D12RHI/Private/D3D12TransientResourceAllocator.cpp
@@ -221,7 +221,7 @@ void FD3D12TransientResourceAllocator::DeallocateMemory(FRHITransientBuffer* InB
 	Allocator.DeallocateMemory(InBuffer, InPassIndex);
 }
 
-void FD3D12TransientResourceAllocator::Freeze(FRHICommandListImmediate& RHICmdList)
+void FD3D12TransientResourceAllocator::Freeze(FRHICommandListImmediate& RHICmdList, FRHITransientHeapStats& OutHeapStats)
 {
-	Allocator.Freeze(RHICmdList);
+	Allocator.Freeze(RHICmdList, OutHeapStats);
 }
diff --git a/Engine/Source/Runtime/D3D12RHI/Private/D3D12TransientResourceAllocator.h b/Engine/Source/Runtime/D3D12RHI/Private/D3D12TransientResourceAllocator.h
--- a/Engine/Source/Runtime/D3D12RHI/Private/D3D12TransientResourceAllocator.h
+++ b/Engine/Source/Runtime/D3D12RHI/Private/D3D12TransientResourceAllocator.h
@@ -50,7 +50,7 @@ public:
 	FRHITransientBuffer* CreateBuffer(const FRHIBufferCreateInfo& InCreateInfo, const TCHAR* InDebugName, uint32 InPassIndex) override;
 	void DeallocateMemory(FRHITransientTexture* InTexture, uint32 InPassIndex) override;
 	void DeallocateMemory(FRHITransientBuffer* InBuffer, uint32 InPassIndex) override;
-	void Freeze(FRHICommandListImmediate& RHICmdList) override;
+	void Freeze(FRHICommandListImmediate& RHICmdList, FRHITransientHeapStats& OutHeapStats) override;
 
 private:
 
diff --git a/Engine/Source/Runtime/RHI/Private/RHI.cpp b/Engine/Source/Runtime/RHI/Private/RHI.cpp
--- a/Engine/Source/Runtime/RHI/Private/RHI.cpp
+++ b/Engine/Source/Runtime/RHI/Private/RHI.cpp
@@ -2723,9 +2723,9 @@ void FRHIBufferViewCache::SetDebugName(const TCHAR* DebugName)
 
 #endif
 
-void FRHITransientTexture::Init(const TCHAR* InName, uint32 InAllocationIndex, uint32 InPassIndex)
+void FRHITransientTexture::Init(const TCHAR* InName, uint32 InAllocationIndex, uint32 InPassIndex, const FRHITransientResourceStats& InStats)
 {
-	FRHITransientResource::Init(InName, InAllocationIndex, InPassIndex);
+	FRHITransientResource::Init(InName, InAllocationIndex, InPassIndex, InStats);
 	ViewCache.SetDebugName(InName);
 
 #if !(UE_BUILD_SHIPPING || UE_BUILD_TEST)
@@ -2733,9 +2733,9 @@ void FRHITransientTexture::Init(const TCHAR* InName, uint32 InAllocationIndex, u
 #endif
 }
 
-void FRHITransientBuffer::Init(const TCHAR* InName, uint32 InAllocationIndex, uint32 InPassIndex)
+void FRHITransientBuffer::Init(const TCHAR* InName, uint32 InAllocationIndex, uint32 InPassIndex, const FRHITransientResourceStats& InStats)
 {
-	FRHITransientResource::Init(InName, InAllocationIndex, InPassIndex);
+	FRHITransientResource::Init(InName, InAllocationIndex, InPassIndex, InStats);
 	ViewCache.SetDebugName(InName);
 
 	// TODO: Add method to rename a buffer.
diff --git a/Engine/Source/Runtime/RHI/Private/RHIValidation.cpp b/Engine/Source/Runtime/RHI/Private/RHIValidation.cpp
--- a/Engine/Source/Runtime/RHI/Private/RHIValidation.cpp
+++ b/Engine/Source/Runtime/RHI/Private/RHIValidation.cpp
@@ -1643,7 +1643,7 @@ void FValidationTransientResourceAllocator::Flush(FRHICommandListImmediate& RHIC
 	RHIAllocator->Flush(RHICmdList);
 }
 
-void FValidationTransientResourceAllocator::Freeze(FRHICommandListImmediate& RHICmdList)
+void FValidationTransientResourceAllocator::Freeze(FRHICommandListImmediate& RHICmdList, FRHITransientHeapStats& OutHeapStats)
 {
 	check(!bFrozen);
 
@@ -1685,7 +1685,7 @@ void FValidationTransientResourceAllocator::Freeze(FRHICommandListImmediate& RHI
 
 	bFrozen = true;
 
-	RHIAllocator->Freeze(RHICmdList);
+	RHIAllocator->Freeze(RHICmdList, OutHeapStats);
 }
 
 void FValidationTransientResourceAllocator::Release(FRHICommandListImmediate& RHICmdList)
diff --git a/Engine/Source/Runtime/RHI/Public/RHITransientResourceAllocator.h b/Engine/Source/Runtime/RHI/Public/RHITransientResourceAllocator.h
--- a/Engine/Source/Runtime/RHI/Public/RHITransientResourceAllocator.h
+++ b/Engine/Source/Runtime/RHI/Public/RHITransientResourceAllocator.h
@@ -13,6 +13,13 @@ enum class ERHITransientResourceType
 	Buffer
 };
 
+struct FRHITransientResourceStats
+{
+	uint64 HeapOffsetMin = 0;
+	uint64 HeapOffsetMax = 0;
+	uint32 HeapIndex = 0;
+};
+
 class FRHITransientResource
 {
 public:
@@ -25,13 +32,14 @@ public:
 	virtual ~FRHITransientResource() = default;
 
 	// (Internal) Initializes the transient resource with a new allocation / name.
-	virtual void Init(const TCHAR* InName, uint32 InAllocationIndex, uint32 InAcquirePassIndex)
+	virtual void Init(const TCHAR* InName, uint32 InAllocationIndex, uint32 InAcquirePassIndex, const FRHITransientResourceStats& InStats)
 	{
 		Name = InName;
 		AllocationIndex = InAllocationIndex;
 		AcquirePasses = TInterval<uint32>(0, InAcquirePassIndex);
 		DiscardPasses = TInterval<uint32>(0, TNumericLimits<uint32>::Max());
 		AliasingOverlaps.Reset();
+		Stats = InStats;
 	}
 
 	// (Internal) Assigns the discard pass index.
@@ -71,6 +79,8 @@ public:
 	FORCEINLINE bool IsTexture() const { return Type == ERHITransientResourceType::Texture; }
 	FORCEINLINE bool IsBuffer() const { return Type == ERHITransientResourceType::Buffer; }
 
+	FORCEINLINE const FRHITransientResourceStats& GetStats() const { return Stats; }
+
 private:
 	// Underlying RHI resource.
 	TRefCountPtr<FRHIResource> Resource;
@@ -91,6 +101,8 @@ private:
 	uint32 AllocationIndex = ~0u;
 	TInterval<uint32> AcquirePasses = TInterval<uint32>(0, 0);
 	TInterval<uint32> DiscardPasses = TInterval<uint32>(0, 0);
+
+	FRHITransientResourceStats Stats;
 };
 
 class RHI_API FRHITransientTexture final : public FRHITransientResource
@@ -101,7 +113,7 @@ public:
 		, CreateInfo(InCreateInfo)
 	{}
 
-	void Init(const TCHAR* InName, uint32 InAllocationIndex, uint32 InPassIndex) override;
+	void Init(const TCHAR* InName, uint32 InAllocationIndex, uint32 InPassIndex, const FRHITransientResourceStats& InStats) override;
 
 	// Returns the underlying RHI texture.
 	FRHITexture* GetRHI() const { return static_cast<FRHITexture*>(FRHITransientResource::GetRHI()); }
@@ -130,7 +142,7 @@ public:
 		, CreateInfo(InCreateInfo)
 	{}
 
-	void Init(const TCHAR* InName, uint32 InAllocationIndex, uint32 InPassIndex) override;
+	void Init(const TCHAR* InName, uint32 InAllocationIndex, uint32 InPassIndex, const FRHITransientResourceStats& InStats) override;
 
 	// Returns the underlying RHI buffer.
 	FORCEINLINE FRHIBuffer* GetRHI() const { return static_cast<FRHIBuffer*>(FRHITransientResource::GetRHI()); }
@@ -151,6 +163,20 @@ public:
 	FRHIBufferViewCache ViewCache;
 };
 
+struct FRHITransientHeapStats
+{
+	struct FHeap
+	{
+		// Maximum number of bytes used in the heap.
+		uint64 WatermarkSize{};
+
+		// Number of bytes allocated for use in the heap.
+		uint64 Capacity{};
+	};
+
+	TArray<FHeap, TInlineAllocator<4>> Heaps;
+};
+
 class IRHITransientResourceAllocator
 {
 public:
@@ -168,7 +194,7 @@ public:
 	virtual void Flush(FRHICommandListImmediate& RHICmdList) {}
 
 	// Freezes all allocations and validates that all resources have their memory deallocated.
-	virtual void Freeze(FRHICommandListImmediate& RHICmdList) = 0;
+	virtual void Freeze(FRHICommandListImmediate& RHICmdList, FRHITransientHeapStats& OutHeapStats) = 0;
 
 	// Releases the transient allocator and deletes the instance. Any FRHITransientResource* access after this call is not allowed.
 	virtual void Release(FRHICommandListImmediate& RHICmdList)
diff --git a/Engine/Source/Runtime/RHI/Public/RHIValidationTransientResourceAllocator.h b/Engine/Source/Runtime/RHI/Public/RHIValidationTransientResourceAllocator.h
--- a/Engine/Source/Runtime/RHI/Public/RHIValidationTransientResourceAllocator.h
+++ b/Engine/Source/Runtime/RHI/Public/RHIValidationTransientResourceAllocator.h
@@ -21,7 +21,7 @@ public:
 	virtual void DeallocateMemory(FRHITransientTexture* InTexture, uint32 InPassIndex) override final;
 	virtual void DeallocateMemory(FRHITransientBuffer* InBuffer, uint32 InPassIndex) override final;
 	virtual void Flush(FRHICommandListImmediate&) override final;
-	virtual void Freeze(FRHICommandListImmediate&) override final;
+	virtual void Freeze(FRHICommandListImmediate&, FRHITransientHeapStats&) override final;
 	virtual void Release(FRHICommandListImmediate&) override final;
 
 private:
diff --git a/Engine/Source/Runtime/RHICore/Private/RHICoreTransientResourceAllocator.cpp b/Engine/Source/Runtime/RHICore/Private/RHICoreTransientResourceAllocator.cpp
--- a/Engine/Source/Runtime/RHICore/Private/RHICoreTransientResourceAllocator.cpp
+++ b/Engine/Source/Runtime/RHICore/Private/RHICoreTransientResourceAllocator.cpp
@@ -612,8 +612,15 @@ void FRHITransientResourceAllocator::DeallocateMemoryInternal(FRHITransientResou
 	StatsToUpdate.CurrentSizeWithAliasing -= Allocation.Size;
 }
 
-void FRHITransientResourceAllocator::Freeze(FRHICommandListImmediate& RHICmdList)
+void FRHITransientResourceAllocator::Freeze(FRHICommandListImmediate& RHICmdList, FRHITransientHeapStats& OutHeapStats)
 {
+	for (FRHITransientHeapState& State : Heaps)
+	{
+		auto& Heap = OutHeapStats.Heaps.AddDefaulted_GetRef();
+		Heap.WatermarkSize = State.CommitSize;
+		Heap.Capacity = State.Heap->GetCapacity();
+	}
+
 	ParentSystem.ForfeitHeaps(Heaps);
 
 	RHICmdList.EnqueueLambda([&ParentSystem = ParentSystem, BufferStats = BufferStats, TextureStats = TextureStats](FRHICommandListImmediate&)
@@ -679,7 +686,12 @@ FRHITransientHeapAllocation FRHITransientResourceAllocator::Allocate(FMemoryStat
 
 void FRHITransientResourceAllocator::InitResource(FRHITransientResource* TransientResource, uint32 PassIndex, const FRHITransientHeapAllocation& Allocation, const TCHAR* Name)
 {
-	TransientResource->Init(Name, HeapAllocations.Emplace(Allocation), PassIndex);
+	FRHITransientResourceStats Stats;
+	Stats.HeapIndex = Allocation.HeapIndex;
+	Stats.HeapOffsetMin = Allocation.Offset;
+	Stats.HeapOffsetMax = Allocation.Offset + Allocation.Size;
+
+	TransientResource->Init(Name, HeapAllocations.Emplace(Allocation), PassIndex, Stats);
 
 	HeapAllocators[Allocation.HeapIndex].TrackOverlap(TransientResource, Allocation);
 }
\ No newline at end of file
diff --git a/Engine/Source/Runtime/RHICore/Public/RHICoreTransientResourceAllocator.h b/Engine/Source/Runtime/RHICore/Public/RHICoreTransientResourceAllocator.h
--- a/Engine/Source/Runtime/RHICore/Public/RHICoreTransientResourceAllocator.h
+++ b/Engine/Source/Runtime/RHICore/Public/RHICoreTransientResourceAllocator.h
@@ -602,7 +602,7 @@ public:
 	}
 
 	// Called to signify all allocations have completed. Forfeits all resources / heaps back to the parent system.
-	void Freeze(FRHICommandListImmediate& RHICmdList);
+	void Freeze(FRHICommandListImmediate& RHICmdList, FRHITransientHeapStats& OutHeapStats);
 
 	// Called to flush any active allocations prior to Freeze.
 	void Flush();
diff --git a/Engine/Source/Runtime/RenderCore/Private/RenderGraphBuilder.cpp b/Engine/Source/Runtime/RenderCore/Private/RenderGraphBuilder.cpp
--- a/Engine/Source/Runtime/RenderCore/Private/RenderGraphBuilder.cpp
+++ b/Engine/Source/Runtime/RenderCore/Private/RenderGraphBuilder.cpp
@@ -1589,7 +1589,9 @@ void FRDGBuilder::Execute(EExecuteMode ExecuteMode)
 
 				if (TransientResourceAllocator)
 				{
-					TransientResourceAllocator->Freeze(RHICmdList);
+					FRHITransientHeapStats TransientHeapStats;
+					TransientResourceAllocator->Freeze(RHICmdList, TransientHeapStats);
+					IF_RDG_ENABLE_TRACE(Trace.SetTransientHeapStats(TransientHeapStats));
 				}
 			}
 			else
diff --git a/Engine/Source/Runtime/RenderCore/Private/RenderGraphTrace.cpp b/Engine/Source/Runtime/RenderCore/Private/RenderGraphTrace.cpp
--- a/Engine/Source/Runtime/RenderCore/Private/RenderGraphTrace.cpp
+++ b/Engine/Source/Runtime/RenderCore/Private/RenderGraphTrace.cpp
@@ -14,6 +14,8 @@ UE_TRACE_EVENT_BEGIN(RDGTrace, GraphMessage)
 	UE_TRACE_EVENT_FIELD(uint64, StartCycles)
 	UE_TRACE_EVENT_FIELD(uint64, EndCycles)
 	UE_TRACE_EVENT_FIELD(uint16, PassCount)
+	UE_TRACE_EVENT_FIELD(uint64[], TransientHeapWatermarkSizes)
+	UE_TRACE_EVENT_FIELD(uint64[], TransientHeapCapacities)
 UE_TRACE_EVENT_END()
 
 UE_TRACE_EVENT_BEGIN(RDGTrace, GraphEndMessage)
@@ -50,6 +52,9 @@ UE_TRACE_EVENT_BEGIN(RDGTrace, BufferMessage)
 	UE_TRACE_EVENT_FIELD(uint16, NextOwnerHandle)
 	UE_TRACE_EVENT_FIELD(uint16, Order)
 	UE_TRACE_EVENT_FIELD(uint16[], Passes)
+	UE_TRACE_EVENT_FIELD(uint16, TransientHeapIndex)
+	UE_TRACE_EVENT_FIELD(uint64, TransientHeapOffsetMin)
+	UE_TRACE_EVENT_FIELD(uint64, TransientHeapOffsetMax)
 	UE_TRACE_EVENT_FIELD(bool, IsExternal)
 	UE_TRACE_EVENT_FIELD(bool, IsExtracted)
 	UE_TRACE_EVENT_FIELD(bool, IsCulled)
@@ -64,6 +69,9 @@ UE_TRACE_EVENT_BEGIN(RDGTrace, TextureMessage)
 	UE_TRACE_EVENT_FIELD(uint16, NextOwnerHandle)
 	UE_TRACE_EVENT_FIELD(uint16, Order)
 	UE_TRACE_EVENT_FIELD(uint16[], Passes)
+	UE_TRACE_EVENT_FIELD(uint16, TransientHeapIndex)
+	UE_TRACE_EVENT_FIELD(uint64, TransientHeapOffsetMin)
+	UE_TRACE_EVENT_FIELD(uint64, TransientHeapOffsetMax)
 	UE_TRACE_EVENT_FIELD(uint64, SizeInBytes)
 	UE_TRACE_EVENT_FIELD(uint64, CreateFlags)
 	UE_TRACE_EVENT_FIELD(uint32, Dimension)
@@ -124,11 +132,22 @@ void FRDGTrace::OutputGraphEnd(const FRDGBuilder& GraphBuilder)
 	{
 		const TCHAR* Name = GraphBuilder.BuilderName.GetTCHAR();
 
+		TArray<uint64, TInlineAllocator<8>> TransientHeapWatermarkSizes;
+		TArray<uint64, TInlineAllocator<8>> TransientHeapCapacities;
+
+		for (const auto& Heap : TransientHeapStats.Heaps)
+		{
+			TransientHeapWatermarkSizes.Emplace(Heap.WatermarkSize);
+			TransientHeapCapacities.Emplace(Heap.Capacity);
+		}
+
 		UE_TRACE_LOG(RDGTrace, GraphMessage, RDGChannel)
 			<< GraphMessage.Name(Name, uint16(FCString::Strlen(Name)))
 			<< GraphMessage.StartCycles(GraphStartCycles)
 			<< GraphMessage.EndCycles(FPlatformTime::Cycles64())
-			<< GraphMessage.PassCount(uint16(Passes.Num()));
+			<< GraphMessage.PassCount(uint16(Passes.Num()))
+			<< GraphMessage.TransientHeapWatermarkSizes(TransientHeapWatermarkSizes.GetData(), (uint16)TransientHeapWatermarkSizes.Num())
+			<< GraphMessage.TransientHeapCapacities(TransientHeapCapacities.GetData(), (uint16)TransientHeapCapacities.Num());
 	}
 
 	for (FRDGPassHandle Handle = Passes.Begin(); Handle != Passes.End(); ++Handle)
@@ -239,12 +258,22 @@ void FRDGTrace::OutputGraphEnd(const FRDGBuilder& GraphBuilder)
 			SizeInBytes = RHIComputeMemorySize(TextureRHI);
 		}
 
+		FRHITransientResourceStats TransientStats;
+
+		if (Texture->bTransient)
+		{
+			TransientStats = Texture->TransientTexture->GetStats();
+		}
+
 		UE_TRACE_LOG(RDGTrace, TextureMessage, RDGChannel)
 			<< TextureMessage.Name(Texture->Name, uint16(FCString::Strlen(Texture->Name)))
 			<< TextureMessage.Handle(Handle.GetIndex())
 			<< TextureMessage.NextOwnerHandle(Texture->NextOwner.GetIndexUnchecked())
 			<< TextureMessage.Order(Texture->TraceOrder)
 			<< TextureMessage.Passes((const uint16*)Texture->TracePasses.GetData(), (uint16)Texture->TracePasses.Num())
+			<< TextureMessage.TransientHeapIndex(TransientStats.HeapIndex)
+			<< TextureMessage.TransientHeapOffsetMin(TransientStats.HeapOffsetMin)
+			<< TextureMessage.TransientHeapOffsetMax(TransientStats.HeapOffsetMax)
 			<< TextureMessage.SizeInBytes(SizeInBytes)
 			<< TextureMessage.CreateFlags(uint32(Texture->Desc.Flags))
 			<< TextureMessage.Dimension(uint32(Texture->Desc.Dimension))
@@ -265,12 +294,22 @@ void FRDGTrace::OutputGraphEnd(const FRDGBuilder& GraphBuilder)
 	{
 		const FRDGBuffer* Buffer = Buffers[Handle];
 
+		FRHITransientResourceStats TransientStats;
+
+		if (Buffer->bTransient)
+		{
+			TransientStats = Buffer->TransientBuffer->GetStats();
+		}
+
 		UE_TRACE_LOG(RDGTrace, BufferMessage, RDGChannel)
 			<< BufferMessage.Name(Buffer->Name, uint16(FCString::Strlen(Buffer->Name)))
 			<< BufferMessage.Handle(Buffer->Handle.GetIndex())
 			<< BufferMessage.NextOwnerHandle(Buffer->NextOwner.GetIndexUnchecked())
 			<< BufferMessage.Order(Buffer->TraceOrder)
 			<< BufferMessage.Passes((const uint16*)Buffer->TracePasses.GetData(), (uint16)Buffer->TracePasses.Num())
+			<< BufferMessage.TransientHeapIndex(TransientStats.HeapIndex)
+			<< BufferMessage.TransientHeapOffsetMin(TransientStats.HeapOffsetMin)
+			<< BufferMessage.TransientHeapOffsetMax(TransientStats.HeapOffsetMax)
 			<< BufferMessage.UsageFlags(uint32(Buffer->Desc.Usage))
 			<< BufferMessage.BytesPerElement(Buffer->Desc.BytesPerElement)
 			<< BufferMessage.NumElements(Buffer->Desc.NumElements)
diff --git a/Engine/Source/Runtime/RenderCore/Public/RenderGraphTrace.h b/Engine/Source/Runtime/RenderCore/Public/RenderGraphTrace.h
--- a/Engine/Source/Runtime/RenderCore/Public/RenderGraphTrace.h
+++ b/Engine/Source/Runtime/RenderCore/Public/RenderGraphTrace.h
@@ -3,6 +3,7 @@
 #pragma once
 
 #include "RenderGraphDefinitions.h"
+#include "RHITransientResourceAllocator.h"
 
 #if RDG_ENABLE_TRACE
 
@@ -18,7 +19,10 @@ public:
 	void AddTexturePassDependency(FRDGTexture* Texture, FRDGPass* Pass);
 	void AddBufferPassDependency(FRDGBuffer* Buffer, FRDGPass* Pass);
 
+	void SetTransientHeapStats(const FRHITransientHeapStats& Stats) { TransientHeapStats = Stats; }
+
 private:
+	FRHITransientHeapStats TransientHeapStats;
 	uint64 GraphStartCycles{};
 	uint32 ResourceOrder{};
 };
