Fix duplicate localization key.

#rnx

diff --git a/Engine/Plugins/RenderGraphInsights/Source/RenderGraphInsights/Private/RenderGraphTrack.cpp b/Engine/Plugins/RenderGraphInsights/Source/RenderGraphInsights/Private/RenderGraphTrack.cpp
--- a/Engine/Plugins/RenderGraphInsights/Source/RenderGraphInsights/Private/RenderGraphTrack.cpp
+++ b/Engine/Plugins/RenderGraphInsights/Source/RenderGraphInsights/Private/RenderGraphTrack.cpp
@@ -1694,7 +1694,7 @@ void FRenderGraphTrack::BuildContextMenu(FMenuBuilder& MenuBuilder)
 		MenuBuilder.AddMenuEntry
 		(
 			LOCTEXT("TransientHeap", "Transient Heaps"),
-			LOCTEXT("ShowTextures_Tooltip", "Shows resource placement in transient heaps."),
+			LOCTEXT("TransientHeap_Tooltip", "Shows resource placement in transient heaps."),
 			FSlateIcon(),
 			FUIAction(
 				FExecuteAction::CreateLambda([this, TimingViewSession]()
