#UE A couple of monolithic fixes for modules used by UnrealInsights

diff --git a/Engine/Plugins/RenderGraphInsights/Source/RenderGraphInsights/Private/RenderGraphInsightsModule.cpp b/Engine/Plugins/RenderGraphInsights/Source/RenderGraphInsights/Private/RenderGraphInsightsModule.cpp
--- a/Engine/Plugins/RenderGraphInsights/Source/RenderGraphInsights/Private/RenderGraphInsightsModule.cpp
+++ b/Engine/Plugins/RenderGraphInsights/Source/RenderGraphInsights/Private/RenderGraphInsightsModule.cpp
@@ -32,4 +32,4 @@ void FRenderGraphInsightsModule::ShutdownModule()
 } //namespace SlateInsights
 } //namespace UE
 
-IMPLEMENT_MODULE((UE::RenderGraphInsights::FRenderGraphInsightsModule), RenderGraphInsights);
\ No newline at end of file
+IMPLEMENT_MODULE(UE::RenderGraphInsights::FRenderGraphInsightsModule, RenderGraphInsights);
\ No newline at end of file
diff --git a/Engine/Plugins/Slate/SlateInsights/Source/SlateInsights/Private/SlateInsightsModule.cpp b/Engine/Plugins/Slate/SlateInsights/Source/SlateInsights/Private/SlateInsightsModule.cpp
--- a/Engine/Plugins/Slate/SlateInsights/Source/SlateInsights/Private/SlateInsightsModule.cpp
+++ b/Engine/Plugins/Slate/SlateInsights/Source/SlateInsights/Private/SlateInsightsModule.cpp
@@ -105,6 +105,6 @@ void FSlateInsightsModule::RegisterTimingProfilerLayoutExtensions(FInsightsMajor
 } //namespace UE
 
 
-IMPLEMENT_MODULE((UE::SlateInsights::FSlateInsightsModule), SlateInsights);
+IMPLEMENT_MODULE(UE::SlateInsights::FSlateInsightsModule, SlateInsights);
 
 #undef LOCTEXT_NAMESPACE
