Fix typos in tooltips.

#rnx

diff --git a/Engine/Plugins/RenderGraphInsights/Source/RenderGraphInsights/Private/RenderGraphTrack.cpp b/Engine/Plugins/RenderGraphInsights/Source/RenderGraphInsights/Private/RenderGraphTrack.cpp
--- a/Engine/Plugins/RenderGraphInsights/Source/RenderGraphInsights/Private/RenderGraphTrack.cpp
+++ b/Engine/Plugins/RenderGraphInsights/Source/RenderGraphInsights/Private/RenderGraphTrack.cpp
@@ -2007,7 +2007,7 @@ void FRenderGraphTrack::BuildContextMenu(FMenuBuilder& MenuBuilder)
 		MenuBuilder.AddMenuEntry
 		(
 			LOCTEXT("4MiB", "4MiB"),
-			LOCTEXT("4MiB_Tooltip", "Each depth slot in the timing view is 1MiB in size on the vertical axis."),
+			LOCTEXT("4MiB_Tooltip", "Each depth slot in the timing view is 4MiB in size on the vertical axis."),
 			FSlateIcon(),
 			FUIAction(
 				FExecuteAction::CreateLambda([this, MiB]()
@@ -2025,7 +2025,7 @@ void FRenderGraphTrack::BuildContextMenu(FMenuBuilder& MenuBuilder)
 		MenuBuilder.AddMenuEntry
 		(
 			LOCTEXT("2MiB", "2MiB"),
-			LOCTEXT("2MiB_Tooltip", "Each depth slot in the timing view is 1MiB in size on the vertical axis."),
+			LOCTEXT("2MiB_Tooltip", "Each depth slot in the timing view is 2MiB in size on the vertical axis."),
 			FSlateIcon(),
 			FUIAction(
 				FExecuteAction::CreateLambda([this, MiB]()
