Fixed RDG insights crash due to RDG track being used in loading insights tab.

diff --git a/Engine/Plugins/RenderGraphInsights/Source/RenderGraphInsights/Private/RenderGraphTimingViewExtender.cpp b/Engine/Plugins/RenderGraphInsights/Source/RenderGraphInsights/Private/RenderGraphTimingViewExtender.cpp
--- a/Engine/Plugins/RenderGraphInsights/Source/RenderGraphInsights/Private/RenderGraphTimingViewExtender.cpp
+++ b/Engine/Plugins/RenderGraphInsights/Source/RenderGraphInsights/Private/RenderGraphTimingViewExtender.cpp
@@ -2,6 +2,7 @@
 
 #include "RenderGraphTimingViewExtender.h"
 #include "Insights/ITimingViewSession.h"
+#include "Insights/IUnrealInsightsModule.h"
 
 #define LOCTEXT_NAMESPACE "RenderGraphTimingViewExtender"
 
@@ -12,6 +13,11 @@ namespace RenderGraphInsights
 
 void FRenderGraphTimingViewExtender::OnBeginSession(Insights::ITimingViewSession& InSession)
 {
+	if (InSession.GetName() != FInsightsManagerTabs::TimingProfilerTabId)
+	{
+		return;
+	}
+
 	FPerSessionData* PerSessionData = PerSessionDataMap.Find(&InSession);
 	if(PerSessionData == nullptr)
 	{
@@ -28,6 +34,11 @@ void FRenderGraphTimingViewExtender::OnBeginSession(Insights::ITimingViewSession
 
 void FRenderGraphTimingViewExtender::OnEndSession(Insights::ITimingViewSession& InSession)
 {
+	if (InSession.GetName() != FInsightsManagerTabs::TimingProfilerTabId)
+	{
+		return;
+	}
+
 	FPerSessionData* PerSessionData = PerSessionDataMap.Find(&InSession);
 	if(PerSessionData != nullptr)
 	{
