Enable RDG Insights plugin for UnrealInsights by default.

diff --git a/Engine/Plugins/RenderGraphInsights/RenderGraphInsights.uplugin b/Engine/Plugins/RenderGraphInsights/RenderGraphInsights.uplugin
--- a/Engine/Plugins/RenderGraphInsights/RenderGraphInsights.uplugin
+++ b/Engine/Plugins/RenderGraphInsights/RenderGraphInsights.uplugin
@@ -10,7 +10,7 @@
 	"DocsURL": "",
 	"MarketplaceURL": "",
 	"SupportURL": "",
-	"EnabledByDefault": false,
+	"EnabledByDefault": true,
 	"CanContainContent": false,
 	"IsBetaVersion": false,
 	"Installed": false,
