Make RDG insights use resource name as coloration by default.

diff --git a/Engine/Plugins/RenderGraphInsights/Source/RenderGraphInsights/Private/RenderGraphTrack.cpp b/Engine/Plugins/RenderGraphInsights/Source/RenderGraphInsights/Private/RenderGraphTrack.cpp
--- a/Engine/Plugins/RenderGraphInsights/Source/RenderGraphInsights/Private/RenderGraphTrack.cpp
+++ b/Engine/Plugins/RenderGraphInsights/Source/RenderGraphInsights/Private/RenderGraphTrack.cpp
@@ -78,6 +78,19 @@ inline uint32 GetResourceColorByTransientCache(bool bHit)
 	return bHit ? HitColor : MissColor;
 }
 
+inline uint32 GetResourceColorByName(const TCHAR* Str)
+{
+	uint32 Color = 0;
+	if (Str != nullptr)
+	{
+		for (const TCHAR* c = Str; *c; ++c)
+		{
+			Color = (Color + *c) * 0x2c2c57ed;
+		}
+	}
+	return Color | 0xFF000000;
+}
+
 ///////////////////////////////////////////////////////////////////////////////
 
 template <typename VisibleItemType>
@@ -456,19 +469,19 @@ uint32 FRenderGraphTrack::GetTextureColor(const FTexturePacket& Texture, uint64
 	{
 		return kUntrackedColor;
 	}
-
+	if (ResourceColor == EResourceColor::Name)
+	{
+		return GetResourceColorByName(*Texture.Name);
+	}
 	if (ResourceColor == EResourceColor::Type)
 	{
 		return kTextureColor;
 	}
-	else if (ResourceColor == EResourceColor::TransientCache)
+	if (ResourceColor == EResourceColor::TransientCache)
 	{
 		return GetResourceColorByTransientCache(!Texture.bTransient || Texture.bTransientCacheHit);
 	}
-	else
-	{
-		return GetResourceColorBySize(Texture.SizeInBytes, MaxSizeInBytes);
-	}
+	return GetResourceColorBySize(Texture.SizeInBytes, MaxSizeInBytes);
 }
 
 uint32 FRenderGraphTrack::GetBufferColor(const FBufferPacket& Buffer, uint64 MaxSizeInBytes) const
@@ -477,19 +490,19 @@ uint32 FRenderGraphTrack::GetBufferColor(const FBufferPacket& Buffer, uint64 Max
 	{
 		return kUntrackedColor;
 	}
-
+	if (ResourceColor == EResourceColor::Name)
+	{
+		return GetResourceColorByName(*Buffer.Name);
+	}
 	if (ResourceColor == EResourceColor::Type)
 	{
 		return kBufferColor;
 	}
-	else if (ResourceColor == EResourceColor::TransientCache)
+	if (ResourceColor == EResourceColor::TransientCache)
 	{
 		return GetResourceColorByTransientCache(!Buffer.bTransient || Buffer.bTransientCacheHit);
 	}
-	else
-	{
-		return GetResourceColorBySize(Buffer.SizeInBytes, MaxSizeInBytes);
-	}
+	return GetResourceColorBySize(Buffer.SizeInBytes, MaxSizeInBytes);
 }
 
 float FRenderGraphTrack::TransientByteOffsetToDepth(uint64 MemoryOffset) const
@@ -1704,6 +1717,24 @@ void FRenderGraphTrack::BuildContextMenu(FMenuBuilder& MenuBuilder)
 	
 	MenuBuilder.BeginSection("Color", LOCTEXT("ColorMenuHeader", "Track Resource Coloration"));
 
+	MenuBuilder.AddMenuEntry
+	(
+		LOCTEXT("ColorName", "By Name"),
+		LOCTEXT("ColorName_Tooltip", "Each resource name has a unique color."),
+		FSlateIcon(),
+		FUIAction(
+			FExecuteAction::CreateLambda([this]()
+			{
+				ResourceColor = EResourceColor::Name;
+				SetDirtyFlag();
+			}),
+			FCanExecuteAction(),
+			FIsActionChecked::CreateLambda([this]() { return ResourceColor == EResourceColor::Name; })
+		),
+		NAME_None,
+		EUserInterfaceActionType::RadioButton
+	);
+
 	MenuBuilder.AddMenuEntry
 	(
 		LOCTEXT("ColorType", "By Type"),
@@ -1739,7 +1770,7 @@ void FRenderGraphTrack::BuildContextMenu(FMenuBuilder& MenuBuilder)
 		NAME_None,
 		EUserInterfaceActionType::RadioButton
 	);
-		
+
 	MenuBuilder.AddMenuEntry
 	(
 		LOCTEXT("ColorTransientCache", "By Transient Cache"),
diff --git a/Engine/Plugins/RenderGraphInsights/Source/RenderGraphInsights/Private/RenderGraphTrack.h b/Engine/Plugins/RenderGraphInsights/Source/RenderGraphInsights/Private/RenderGraphTrack.h
--- a/Engine/Plugins/RenderGraphInsights/Source/RenderGraphInsights/Private/RenderGraphTrack.h
+++ b/Engine/Plugins/RenderGraphInsights/Source/RenderGraphInsights/Private/RenderGraphTrack.h
@@ -312,6 +312,7 @@ private:
 
 	enum class EResourceColor
 	{
+		Name,
 		Type,
 		Size,
 		TransientCache
@@ -325,7 +326,7 @@ private:
 
 	EResourceShow ResourceShow = EResourceShow::All;
 	EResourceSort ResourceSort = EResourceSort::Creation;
-	EResourceColor ResourceColor = EResourceColor::Type;
+	EResourceColor ResourceColor = EResourceColor::Name;
 	EVisualizer Visualizer = EVisualizer::Resources;
 	FString FilterText;
 	float FilterSize{};
