[Insights] Fixed crash in RenderGraphInsights plugin. Added support for late connect (i.e. enabling late the RDGTrace channel).

diff --git a/Engine/Plugins/RenderGraphInsights/Source/RenderGraphInsights/Private/RenderGraphAnalyzer.cpp b/Engine/Plugins/RenderGraphInsights/Source/RenderGraphInsights/Private/RenderGraphAnalyzer.cpp
--- a/Engine/Plugins/RenderGraphInsights/Source/RenderGraphInsights/Private/RenderGraphAnalyzer.cpp
+++ b/Engine/Plugins/RenderGraphInsights/Source/RenderGraphInsights/Private/RenderGraphAnalyzer.cpp
@@ -14,6 +14,7 @@ namespace RenderGraphInsights
 FRenderGraphAnalyzer::FRenderGraphAnalyzer(TraceServices::IAnalysisSession& InSession, FRenderGraphProvider& InRenderGraphProvider)
 	: Session(InSession)
 	, Provider(InRenderGraphProvider)
+	, bIsValidGraph(false)
 {}
 
 void FRenderGraphAnalyzer::OnAnalysisBegin(const FOnAnalysisContext& Context)
@@ -26,6 +27,8 @@ void FRenderGraphAnalyzer::OnAnalysisBegin(const FOnAnalysisContext& Context)
 	Builder.RouteEvent(RouteId_Pass, "RDGTrace", "PassMessage");
 	Builder.RouteEvent(RouteId_Texture, "RDGTrace", "TextureMessage");
 	Builder.RouteEvent(RouteId_Buffer, "RDGTrace", "BufferMessage");
+
+	bIsValidGraph = false;
 }
 
 bool FRenderGraphAnalyzer::OnEvent(uint16 RouteId, EStyle Style, const FOnEventContext& Context)
@@ -37,25 +40,43 @@ bool FRenderGraphAnalyzer::OnEvent(uint16 RouteId, EStyle Style, const FOnEventC
 	{
 	case RouteId_Graph:
 	{
+		ensure(!bIsValidGraph);
 		double EndTime{};
 		Provider.AddGraph(Context, EndTime);
 		Session.UpdateDurationSeconds(EndTime);
+		bIsValidGraph = true;
 		break;
 	}
 	case RouteId_GraphEnd:
-		Provider.AddGraphEnd();
+		if (bIsValidGraph)
+		{
+			Provider.AddGraphEnd();
+			bIsValidGraph = false;
+		}
 		break;
 	case RouteId_Scope:
-		Provider.AddScope(FScopePacket(Context));
+		if (bIsValidGraph)
+		{
+			Provider.AddScope(FScopePacket(Context));
+		}
 		break;
 	case RouteId_Pass:
-		Provider.AddPass(FPassPacket(Context));
+		if (bIsValidGraph)
+		{
+			Provider.AddPass(FPassPacket(Context));
+		}
 		break;
 	case RouteId_Texture:
-		Provider.AddTexture(FTexturePacket(Context));
+		if (bIsValidGraph)
+		{
+			Provider.AddTexture(FTexturePacket(Context));
+		}
 		break;
 	case RouteId_Buffer:
-		Provider.AddBuffer(FBufferPacket(Context));
+		if (bIsValidGraph)
+		{
+			Provider.AddBuffer(FBufferPacket(Context));
+		}
 		break;
 	}
 	return true;
diff --git a/Engine/Plugins/RenderGraphInsights/Source/RenderGraphInsights/Private/RenderGraphAnalyzer.h b/Engine/Plugins/RenderGraphInsights/Source/RenderGraphInsights/Private/RenderGraphAnalyzer.h
--- a/Engine/Plugins/RenderGraphInsights/Source/RenderGraphInsights/Private/RenderGraphAnalyzer.h
+++ b/Engine/Plugins/RenderGraphInsights/Source/RenderGraphInsights/Private/RenderGraphAnalyzer.h
@@ -35,6 +35,7 @@ private:
 
 	TraceServices::IAnalysisSession& Session;
 	FRenderGraphProvider& Provider;
+	bool bIsValidGraph;
 };
 
 } //namespace RenderGraphInsights
