Added missing includes for some enums used to call StaticEnum.

diff --git a/Engine/Plugins/Animation/ControlRig/Source/ControlRig/Private/Rigs/RigHierarchyController.cpp b/Engine/Plugins/Animation/ControlRig/Source/ControlRig/Private/Rigs/RigHierarchyController.cpp
--- a/Engine/Plugins/Animation/ControlRig/Source/ControlRig/Private/Rigs/RigHierarchyController.cpp
+++ b/Engine/Plugins/Animation/ControlRig/Source/ControlRig/Private/Rigs/RigHierarchyController.cpp
@@ -7,6 +7,7 @@
 #include "ModularRig.h"
 #include "HelperUtil.h"
 #include "Engine/SkeletalMesh.h"
+#include "Math/Axis.h"
 #include "Misc/ScopeLock.h"
 #include "Units/Execution/RigUnit_PrepareForExecution.h"
 
diff --git a/Engine/Plugins/PCG/Source/PCG/Private/Compute/PCGComputeCommon.cpp b/Engine/Plugins/PCG/Source/PCG/Private/Compute/PCGComputeCommon.cpp
--- a/Engine/Plugins/PCG/Source/PCG/Private/Compute/PCGComputeCommon.cpp
+++ b/Engine/Plugins/PCG/Source/PCG/Private/Compute/PCGComputeCommon.cpp
@@ -27,6 +27,7 @@
 #include "Editor/IPCGEditorModule.h"
 
 #include "DynamicRHI.h"
+#include "PixelFormat.h"
 #include "RHIStats.h"
 #include "RenderResource.h"
 
diff --git a/Engine/Plugins/RenderGraphInsights/Source/RenderGraphInsights/Private/RenderGraphTrack.cpp b/Engine/Plugins/RenderGraphInsights/Source/RenderGraphInsights/Private/RenderGraphTrack.cpp
--- a/Engine/Plugins/RenderGraphInsights/Source/RenderGraphInsights/Private/RenderGraphTrack.cpp
+++ b/Engine/Plugins/RenderGraphInsights/Source/RenderGraphInsights/Private/RenderGraphTrack.cpp
@@ -3,6 +3,7 @@
 #include "RenderGraphTrack.h"
 
 #include "Framework/MultiBox/MultiBoxBuilder.h"
+#include "PixelFormat.h"
 #include "RenderGraphTimingViewSession.h"
 #include "RenderGraphTrackDrawHelper.h"
 #include "Widgets/Input/SSearchBox.h"
diff --git a/Engine/Source/Developer/OutputLog/Private/SOutputLog.cpp b/Engine/Source/Developer/OutputLog/Private/SOutputLog.cpp
--- a/Engine/Source/Developer/OutputLog/Private/SOutputLog.cpp
+++ b/Engine/Source/Developer/OutputLog/Private/SOutputLog.cpp
@@ -9,6 +9,7 @@
 #include "Framework/Text/TextLayout.h"
 #include "Misc/ConfigCacheIni.h"
 #include "Misc/OutputDeviceHelper.h"
+#include "Misc/OutputDevice.h"
 #include "Misc/ScopeLock.h"
 #include "SlateOptMacros.h"
 #include "Textures/SlateIcon.h"
diff --git a/Engine/Source/Runtime/Net/Iris/Private/Iris/ReplicationSystem/ObjectReplicationBridge.cpp b/Engine/Source/Runtime/Net/Iris/Private/Iris/ReplicationSystem/ObjectReplicationBridge.cpp
--- a/Engine/Source/Runtime/Net/Iris/Private/Iris/ReplicationSystem/ObjectReplicationBridge.cpp
+++ b/Engine/Source/Runtime/Net/Iris/Private/Iris/ReplicationSystem/ObjectReplicationBridge.cpp
@@ -4,6 +4,7 @@
 
 #include "Misc/ScopeExit.h"
 #include "HAL/IConsoleManager.h"
+#include "UObject/CoreNetTypes.h"
 
 #include "Iris/IrisConfigInternal.h"
 
